File Sandbox plugin: Adds the FileSandbox plugin. It allows the engine to be put into a sandbox such that you can edit, remove, and add uassets without affecting the project files. You can persist the changes, which means to apply them to your project files.

This change is API only.

#Synced-CL 47637010

diff --git a/Engine/Plugins/Developer/Sandbox/FileSandbox/FileSandbox.uplugin b/Engine/Plugins/Developer/Sandbox/FileSandbox/FileSandbox.uplugin
new file mode 100644
--- /dev/null
+++ b/Engine/Plugins/Developer/Sandbox/FileSandbox/FileSandbox.uplugin
@@ -0,0 +1,27 @@
+{
+	"FileVersion": 3,
+	"Version": 1,
+	"VersionName": "0.1",
+	"FriendlyName": "File Sandbox",
+	"Description": "Core functionality for sandboxing files in the editor.",
+	"Category": "Developer",
+	"CreatedBy": "Epic Games, Inc.",
+	"CreatedByURL": "https://epicgames.com",
+	"DocsURL": "",
+	"MarketplaceURL": "",
+	"SupportURL": "",
+	"EnabledByDefault": false,
+	"CanContainContent": false,
+	"Hidden": true,
+	"IsBetaVersion": true,
+	"Installed": false,
+	"SupportedPrograms": [ "UnrealMultiUserServer", "CoopMultiUserServer", "UnrealMultiUserSlateServer", "UnrealRecoverySvc", "CrashReportClientEditor", "LiveLinkHub" ],
+	"Modules": [
+		{
+			"Name": "FileSandboxCore",
+			"Type": "Runtime",
+			"LoadingPhase": "Default", 
+			"ProgramAllowList": [ "UnrealMultiUserServer", "CoopMultiUserServer", "UnrealMultiUserSlateServer", "UnrealRecoverySvc", "CrashReportClientEditor", "LiveLinkHub" ]
+		}
+	]
+}
diff --git a/Engine/Plugins/Developer/Sandbox/FileSandbox/Source/FileSandboxCore/FileSandboxCore.Build.cs b/Engine/Plugins/Developer/Sandbox/FileSandbox/Source/FileSandboxCore/FileSandboxCore.Build.cs
new file mode 100644
--- /dev/null
+++ b/Engine/Plugins/Developer/Sandbox/FileSandbox/Source/FileSandboxCore/FileSandboxCore.Build.cs
@@ -0,0 +1,39 @@
+// Copyright Epic Games, Inc. All Rights Reserved.
+
+namespace UnrealBuildTool.Rules
+{
+	public class FileSandboxCore : ModuleRules
+	{
+		public FileSandboxCore(ReadOnlyTargetRules Target) : base(Target)
+		{
+			PublicDependencyModuleNames.AddRange(
+				new string[]
+				{
+					"Core",
+					"CoreUObject",  
+				}
+			);
+			PrivateDependencyModuleNames.AddRange(
+				new string[]
+				{
+					"DirectoryWatcher",
+					"Engine",
+					"Json",
+					"JsonUtilities", 
+					"RenderCore",
+					"SourceControl", 
+				});
+
+			if (Target.bBuildEditor)
+			{
+				PrivateDependencyModuleNames.AddRange(
+					new string[]
+					{
+						"EngineSettings",
+						"UnrealEd"
+					}
+				);
+			}
+		}
+	}
+}
diff --git a/Engine/Plugins/Developer/Sandbox/FileSandbox/Source/FileSandboxCore/Private/Data/AbsoluteFilePathArray.h b/Engine/Plugins/Developer/Sandbox/FileSandbox/Source/FileSandboxCore/Private/Data/AbsoluteFilePathArray.h
new file mode 100644
--- /dev/null
+++ b/Engine/Plugins/Developer/Sandbox/FileSandbox/Source/FileSandboxCore/Private/Data/AbsoluteFilePathArray.h
@@ -0,0 +1,41 @@
+﻿// Copyright Epic Games, Inc. All Rights Reserved.
+
+#pragma once
+
+#include "Misc/Paths.h"
+#include "AbsoluteFilePathArray.generated.h"
+
+/** Keeps an array of absolute file paths. Each path is unique. */
+USTRUCT()
+struct FFileSandboxCore_AbsoluteFilePathArray
+{
+	GENERATED_BODY()
+
+	const TArray<FString>& GetFiles() const { return Files; }
+
+	/** @return Whether the path was added, i.e. was not already present. */
+	bool Add(const FString& InPath)
+	{
+		const int32 NumBefore = Files.Num();
+		Files.AddUnique(FPaths::ConvertRelativePathToFull(InPath));
+		return NumBefore != Files.Num();
+	}
+
+	/** @return Whether anything was removed. */
+	bool Remove(const FString& InPath)
+	{
+		const int32 NumRemoved = Files.RemoveSingle(FPaths::ConvertRelativePathToFull(InPath));
+		return NumRemoved > 0;
+	}
+
+	void Empty()
+	{
+		Files.Empty();
+	}
+
+private:
+
+	/** Absolute file paths. */
+	UPROPERTY()
+	TArray<FString> Files;
+};
\ No newline at end of file
diff --git a/Engine/Plugins/Developer/Sandbox/FileSandbox/Source/FileSandboxCore/Private/Data/ManifestData.h b/Engine/Plugins/Developer/Sandbox/FileSandbox/Source/FileSandboxCore/Private/Data/ManifestData.h
new file mode 100644
--- /dev/null
+++ b/Engine/Plugins/Developer/Sandbox/FileSandbox/Source/FileSandboxCore/Private/Data/ManifestData.h
@@ -0,0 +1,32 @@
+// Copyright Epic Games, Inc. All Rights Reserved.
+
+#pragma once
+
+#include "AbsoluteFilePathArray.h"
+#include "ManifestData.generated.h"
+
+/** Internal bookkeeping data for the sandbox logic. */
+USTRUCT()
+struct FFileSandboxCore_ManifestData
+{
+	GENERATED_BODY()
+	
+	/** Absolute non-sandbox paths to files deleted in sandbox. */
+	UPROPERTY()
+	FFileSandboxCore_AbsoluteFilePathArray DeletedFiles;
+
+	/** Absolute non-sandbox paths to files modified in sandbox. */
+	UPROPERTY()
+	FFileSandboxCore_AbsoluteFilePathArray ModifiedFiles;
+
+	/** Absolute non-sandbox paths to files added in sandbox. */
+	UPROPERTY()
+	FFileSandboxCore_AbsoluteFilePathArray AddedFiles;
+
+	void Empty()
+	{
+		DeletedFiles.Empty();
+		ModifiedFiles.Empty();
+		AddedFiles.Empty();
+	}
+};
\ No newline at end of file
diff --git a/Engine/Plugins/Developer/Sandbox/FileSandbox/Source/FileSandboxCore/Private/FileSandboxCoreModule.cpp b/Engine/Plugins/Developer/Sandbox/FileSandbox/Source/FileSandboxCore/Private/FileSandboxCoreModule.cpp
new file mode 100644
--- /dev/null
+++ b/Engine/Plugins/Developer/Sandbox/FileSandbox/Source/FileSandboxCore/Private/FileSandboxCoreModule.cpp
@@ -0,0 +1,19 @@
+// Copyright Epic Games, Inc. All Rights Reserved.
+
+#include "FileSandboxCoreModule.h"
+
+namespace UE::FileSandboxCore
+{
+void FFileSandboxCoreModule::StartupModule()
+{
+	SandboxManager = MakeUnique<FSandboxManager>();
+}
+
+void FFileSandboxCoreModule::ShutdownModule()
+{
+	SandboxManager.Reset();
+}
+}
+
+IMPLEMENT_MODULE(UE::FileSandboxCore::FFileSandboxCoreModule, FileSandboxCore);
+
diff --git a/Engine/Plugins/Developer/Sandbox/FileSandbox/Source/FileSandboxCore/Private/FileSandboxCoreModule.h b/Engine/Plugins/Developer/Sandbox/FileSandbox/Source/FileSandboxCore/Private/FileSandboxCoreModule.h
new file mode 100644
--- /dev/null
+++ b/Engine/Plugins/Developer/Sandbox/FileSandbox/Source/FileSandboxCore/Private/FileSandboxCoreModule.h
@@ -0,0 +1,29 @@
+// Copyright Epic Games, Inc. All Rights Reserved.
+
+#pragma once
+
+#include "IFileSandboxCoreModule.h"
+#include "Sandbox/SandboxManager.h"
+
+namespace UE::FileSandboxCore
+{
+class FFileSandboxCoreModule : public IFileSandboxCoreModule
+{
+public:
+
+	//~ Begin IModuleInterface Interface
+	virtual void StartupModule() override;
+	virtual void ShutdownModule() override;
+	//~ End IModuleInterface Interface
+
+	//~ Begin IFileSandboxCoreModule Module
+	virtual ISandboxManager& GetSandboxManager() override { return *SandboxManager; }
+	//~ End IFileSandboxCoreModule Module
+
+private:
+
+	/** Public API for interacting with the sandbox instance. */
+	TUniquePtr<FSandboxManager> SandboxManager;
+};
+}
+
diff --git a/Engine/Plugins/Developer/Sandbox/FileSandbox/Source/FileSandboxCore/Private/LogFileSandbox.cpp b/Engine/Plugins/Developer/Sandbox/FileSandbox/Source/FileSandboxCore/Private/LogFileSandbox.cpp
new file mode 100644
--- /dev/null
+++ b/Engine/Plugins/Developer/Sandbox/FileSandbox/Source/FileSandboxCore/Private/LogFileSandbox.cpp
@@ -0,0 +1,5 @@
+﻿// Copyright Epic Games, Inc. All Rights Reserved.
+
+#include "LogFileSandbox.h"
+
+DEFINE_LOG_CATEGORY(LogFileSandbox);
\ No newline at end of file
diff --git a/Engine/Plugins/Developer/Sandbox/FileSandbox/Source/FileSandboxCore/Private/LogFileSandbox.h b/Engine/Plugins/Developer/Sandbox/FileSandbox/Source/FileSandboxCore/Private/LogFileSandbox.h
new file mode 100644
--- /dev/null
+++ b/Engine/Plugins/Developer/Sandbox/FileSandbox/Source/FileSandboxCore/Private/LogFileSandbox.h
@@ -0,0 +1,7 @@
+﻿// Copyright Epic Games, Inc. All Rights Reserved.
+
+#pragma once
+
+#include "Logging/LogCategory.h"
+
+DECLARE_LOG_CATEGORY_EXTERN(LogFileSandbox, Log, All);
diff --git a/Engine/Plugins/Developer/Sandbox/FileSandbox/Source/FileSandboxCore/Private/Sandbox/ISandboxInstance.cpp b/Engine/Plugins/Developer/Sandbox/FileSandbox/Source/FileSandboxCore/Private/Sandbox/ISandboxInstance.cpp
new file mode 100644
--- /dev/null
+++ b/Engine/Plugins/Developer/Sandbox/FileSandbox/Source/FileSandboxCore/Private/Sandbox/ISandboxInstance.cpp
@@ -0,0 +1,24 @@
+﻿// Copyright Epic Games, Inc. All Rights Reserved.
+
+#include "ISandboxInstance.h"
+
+#include "Types/Sandbox/DiscardArgs.h"
+#include "Types/Sandbox/DiscardResult.h"
+#include "Types/Sandbox/PersistArgs.h"
+#include "Utils/PackageSandboxUtils.h"
+
+namespace UE::FileSandboxCore
+{
+void ISandboxInstance::PersistAll()
+{
+	const TArray<FString> ChangedFiles = GatherChangedFiles();
+	PersistSandbox(FPersistArgs{ ChangedFiles});
+}
+
+void ISandboxInstance::DiscardAll()
+{
+	const FDiscardResult Result = DiscardSandbox(FDiscardArgs());
+	HotReloadPackages(Result.PackagesPendingHotReload);
+	PurgePackages(Result.PackagesPendingPurge);
+}
+}
diff --git a/Engine/Plugins/Developer/Sandbox/FileSandbox/Source/FileSandboxCore/Private/Sandbox/Platform/SandboxPlatformFile.cpp b/Engine/Plugins/Developer/Sandbox/FileSandbox/Source/FileSandboxCore/Private/Sandbox/Platform/SandboxPlatformFile.cpp
new file mode 100644
--- /dev/null
+++ b/Engine/Plugins/Developer/Sandbox/FileSandbox/Source/FileSandboxCore/Private/Sandbox/Platform/SandboxPlatformFile.cpp
@@ -0,0 +1,1312 @@
+﻿// Copyright Epic Games, Inc. All Rights Reserved.
+
+#include "SandboxPlatformFile.h"
+
+#include "Containers/UnrealString.h"
+#include "Engine/GameEngine.h"
+#include "Engine/Level.h"
+#include "Engine/World.h"
+#include "GenericPlatform/GenericPlatformFile.h"
+#include "HAL/PlatformFileManager.h"
+#include "ISourceControlModule.h"
+#include "ISourceControlProvider.h"
+#include "ISourceControlState.h"
+#include "JsonObjectConverter.h"
+#include "LogFileSandbox.h"
+#include "Logging/LogVerbosity.h"
+#include "Misc/FileHelper.h"
+#include "Misc/MessageDialog.h"
+#include "Misc/PackageName.h"
+#include "Misc/Paths.h"
+#include "Misc/ScopeLock.h"
+#include "Modules/ModuleManager.h"
+#include "RenderingThread.h"
+#include "SourceControlOperations.h"
+#include "Templates/UniquePtr.h"
+#include "Types/Sandbox/PersistArgs.h"
+#include "Types/Sandbox/PersistResult.h"
+#include "UObject/Linker.h"
+#include "UObject/LinkerLoad.h"
+#include "UObject/Package.h"
+#include "Utils/PackageSandboxUtils.h"
+#include "Utils/SandboxUtils.h"
+
+#if WITH_EDITOR
+#include "DirectoryWatcherModule.h"
+#include "Editor.h"
+#include "Editor/UnrealEdEngine.h"
+#include "FileHelpers.h"
+#include "GameMapsSettings.h"
+#include "IDirectoryWatcher.h"
+#include "PackageTools.h"
+#include "Selection.h"
+#include "Subsystems/AssetEditorSubsystem.h"
+#include "Subsystems/EditorActorSubsystem.h"
+#endif
+
+class IDirectoryWatcher;
+
+namespace UE::FileSandboxCore
+{
+FSandboxPlatformFile::FSandboxPlatformFile(const FString& InSandboxRootPath)
+	: RootPath(InSandboxRootPath)
+	, LowerLevel(nullptr)
+	, bSandboxEnabled(false)
+{
+}
+
+FSandboxPlatformFile::~FSandboxPlatformFile()
+{
+	IPlatformFile& CurrentPlatformFile = FPlatformFileManager::Get().GetPlatformFile();
+	if (LowerLevel && this == &CurrentPlatformFile)
+	{
+		FPlatformFileManager::Get().SetPlatformFile(*LowerLevel);
+	}
+
+	FPackageName::OnContentPathMounted().RemoveAll(this);
+	FPackageName::OnContentPathDismounted().RemoveAll(this);
+
+#if WITH_EDITOR
+	if (IDirectoryWatcher* DirectoryWatcher = GetDirectoryWatcherIfLoaded())
+	{
+		for (FSandboxMountPoint& SandboxMountPoint : SandboxMountPoints)
+		{
+			if (SandboxMountPoint.OnDirectoryChangedHandle.IsValid())
+			{
+				DirectoryWatcher->UnregisterDirectoryChangedCallback_Handle(SandboxMountPoint.Path.GetSandboxPath(), SandboxMountPoint.OnDirectoryChangedHandle);
+				SandboxMountPoint.OnDirectoryChangedHandle.Reset();
+			}
+		}
+	}
+#endif
+}
+
+FString FSandboxPlatformFile::GetSandboxRootPath() const
+{
+	return RootPath / TEXT("Sandbox");
+}
+
+void FSandboxPlatformFile::SetSandboxEnabled(bool bInEnabled)
+{
+	bSandboxEnabled = bInEnabled;
+}
+
+bool FSandboxPlatformFile::IsSandboxEnabled() const
+{
+	return bSandboxEnabled;
+}
+
+bool FSandboxPlatformFile::Initialize(IPlatformFile* Inner, const TCHAR* CmdLine)
+{
+	SetLowerLevel(Inner);
+	
+	LowerLevel->CreateDirectoryTree(*GetSandboxRootPath());
+
+	// Set-up the initial set of content mount paths
+	TArray<FString> RootPaths;
+	FPackageName::QueryRootContentPaths(RootPaths);
+	for (const FString& Path : RootPaths)
+	{
+		RegisterContentMountPath(Path, FPackageName::LongPackageNameToFilename(Path));
+	}
+
+	// Watch for new content mount paths
+	FPackageName::OnContentPathMounted().AddRaw(this, &FSandboxPlatformFile::OnContentPathMounted);
+	FPackageName::OnContentPathDismounted().AddRaw(this, &FSandboxPlatformFile::OnContentPathDismounted);
+
+	bSandboxEnabled = true;
+	FPlatformFileManager::Get().SetPlatformFile(*this);
+
+#if WITH_EDITOR
+	SynchronizeAssetRegistry();
+#endif
+
+	return true;
+}
+
+void FSandboxPlatformFile::Tick()
+{
+}
+
+IPlatformFile* FSandboxPlatformFile::GetLowerLevel()
+{
+	return LowerLevel;
+}
+
+void FSandboxPlatformFile::SetLowerLevel(IPlatformFile* NewLowerLevel)
+{
+	check(NewLowerLevel && NewLowerLevel != this);
+	LowerLevel = NewLowerLevel;
+}
+
+const TCHAR* FSandboxPlatformFile::GetName() const
+{
+	return TEXT("FSandboxPlatformFile");
+}
+
+bool FSandboxPlatformFile::FileExists(const TCHAR* Filename)
+{
+	const FSandboxedPlatformFilePath ResolvedPath = ToSandboxPath(Filename);
+
+	if (ResolvedPath.HasSandboxPath())
+	{
+		if (IsPathDeleted(ResolvedPath))
+		{
+			return false;
+		}
+
+		if (LowerLevel->FileExists(*ResolvedPath.GetSandboxPath()))
+		{
+			return true;
+		}
+	}
+
+	return LowerLevel->FileExists(*ResolvedPath.GetNonSandboxPath());
+}
+
+int64 FSandboxPlatformFile::FileSize(const TCHAR* Filename)
+{
+	const FSandboxedPlatformFilePath ResolvedPath = ToSandboxPath(Filename);
+
+	if (ResolvedPath.HasSandboxPath())
+	{
+		if (IsPathDeleted(ResolvedPath))
+		{
+			return -1;
+		}
+
+		if (LowerLevel->FileExists(*ResolvedPath.GetSandboxPath()))
+		{
+			return LowerLevel->FileSize(*ResolvedPath.GetSandboxPath());
+		}
+	}
+
+	return LowerLevel->FileSize(*ResolvedPath.GetNonSandboxPath());
+}
+
+bool FSandboxPlatformFile::DeleteFile(const TCHAR* Filename)
+{
+	const FSandboxedPlatformFilePath ResolvedPath = ToSandboxPath(Filename);
+
+	if (ResolvedPath.HasSandboxPath())
+	{
+		if (IsPathDeleted(ResolvedPath))
+		{
+			return false;
+		}
+
+		if (!LowerLevel->FileExists(*ResolvedPath.GetSandboxPath()) || LowerLevel->DeleteFile(*ResolvedPath.GetSandboxPath()))
+		{
+			SetPathDeleted(ResolvedPath, true);
+			NotifyFileDeleted(ResolvedPath);
+			return true;
+		}
+
+		return false; // Do not attempt to delete the non-sandbox file
+	}
+
+	return LowerLevel->DeleteFile(*ResolvedPath.GetNonSandboxPath());
+}
+
+bool FSandboxPlatformFile::IsReadOnly(const TCHAR* Filename)
+{
+	const FSandboxedPlatformFilePath ResolvedPath = ToSandboxPath(Filename);
+
+	if (ResolvedPath.HasSandboxPath())
+	{
+		if (IsPathDeleted(ResolvedPath))
+		{
+			return false;
+		}
+
+		if (LowerLevel->FileExists(*ResolvedPath.GetSandboxPath()))
+		{
+			return LowerLevel->IsReadOnly(*ResolvedPath.GetSandboxPath());
+		}
+
+		return false; // Can always overwrite missing sandbox files
+	}
+
+	return LowerLevel->IsReadOnly(*ResolvedPath.GetNonSandboxPath());
+}
+
+bool FSandboxPlatformFile::SetReadOnly(const TCHAR* Filename, bool bNewReadOnlyValue)
+{
+	const FSandboxedPlatformFilePath ResolvedPath = ToSandboxPath(Filename);
+
+	if (ResolvedPath.HasSandboxPath())
+	{
+		if (IsPathDeleted(ResolvedPath))
+		{
+			return false;
+		}
+
+		if (LowerLevel->FileExists(*ResolvedPath.GetSandboxPath()))
+		{
+			if (bNewReadOnlyValue)
+			{
+				return true; // Do not allow sandbox files to be made read-only, but don't report failure
+			}
+
+			return LowerLevel->SetReadOnly(*ResolvedPath.GetSandboxPath(), bNewReadOnlyValue);
+		}
+
+		return false; // Do not attempt to modify the non-sandbox file
+	}
+
+	return LowerLevel->SetReadOnly(*ResolvedPath.GetNonSandboxPath(), bNewReadOnlyValue);
+}
+
+bool FSandboxPlatformFile::MoveFile(const TCHAR* To, const TCHAR* From)
+{
+	const FSandboxedPlatformFilePath ResolvedToPath = ToSandboxPath(To);
+	const FSandboxedPlatformFilePath ResolvedFromPath = ToSandboxPath(From);
+
+	if (ResolvedToPath.HasSandboxPath())
+	{
+		// Migrate any existing target file from outside the sandbox so the lower-level will fail to overwrite the existing file
+		MigrateFileToSandbox(ResolvedToPath);
+
+		if (ResolvedFromPath.HasSandboxPath())
+		{
+			// Sandbox -> Sandbox
+			if (IsPathDeleted(ResolvedFromPath))
+			{
+				// Cannot move a deleted file
+				return false;
+			}
+
+			bool bSuccess = false;
+			if (LowerLevel->FileExists(*ResolvedFromPath.GetSandboxPath()))
+			{
+				// Moving an internal sandbox file - can move
+				bSuccess = LowerLevel->MoveFile(*ResolvedToPath.GetSandboxPath(), *ResolvedFromPath.GetSandboxPath());
+			}
+			else
+			{
+				// Moving an external sandbox file - must copy
+				bSuccess = LowerLevel->CopyFile(*ResolvedToPath.GetSandboxPath(), *ResolvedFromPath.GetNonSandboxPath());
+			}
+
+			if (bSuccess)
+			{
+				SetPathDeleted(ResolvedToPath, false);
+				SetPathDeleted(ResolvedFromPath, true);
+				NotifyFileDeleted(ResolvedFromPath);
+			}
+
+			return bSuccess;
+		}
+		else
+		{
+			// Non-sandbox -> Sandbox
+			if (LowerLevel->MoveFile(*ResolvedToPath.GetSandboxPath(), *ResolvedFromPath.GetNonSandboxPath()))
+			{
+				SetPathDeleted(ResolvedToPath, false);
+				return true;
+			}
+			return false;
+		}
+	}
+	else if (ResolvedFromPath.HasSandboxPath())
+	{
+		// Sandbox -> Non-sandbox
+		if (IsPathDeleted(ResolvedFromPath))
+		{
+			// Cannot move a deleted file
+			return false;
+		}
+
+		bool bSuccess = false;
+		if (LowerLevel->FileExists(*ResolvedFromPath.GetSandboxPath()))
+		{
+			// Moving an internal sandbox file - can move
+			bSuccess = LowerLevel->MoveFile(*ResolvedToPath.GetNonSandboxPath(), *ResolvedFromPath.GetSandboxPath());
+		}
+		else
+		{
+			// Moving an external sandbox file - must copy
+			bSuccess = LowerLevel->CopyFile(*ResolvedToPath.GetNonSandboxPath(), *ResolvedFromPath.GetNonSandboxPath());
+		}
+
+		if (bSuccess)
+		{
+			SetPathDeleted(ResolvedFromPath, true);
+			NotifyFileDeleted(ResolvedFromPath);
+		}
+
+		return bSuccess;
+	}
+
+	// Non-sandbox -> Non-sandbox
+	return LowerLevel->MoveFile(*ResolvedToPath.GetNonSandboxPath(), *ResolvedFromPath.GetNonSandboxPath());
+}
+
+FDateTime FSandboxPlatformFile::GetTimeStamp(const TCHAR* Filename)
+{
+	const FSandboxedPlatformFilePath ResolvedPath = ToSandboxPath(Filename);
+
+	if (ResolvedPath.HasSandboxPath())
+	{
+		if (IsPathDeleted(ResolvedPath))
+		{
+			return FDateTime::MinValue();
+		}
+
+		if (LowerLevel->FileExists(*ResolvedPath.GetSandboxPath()))
+		{
+			return LowerLevel->GetTimeStamp(*ResolvedPath.GetSandboxPath());
+		}
+	}
+
+	return LowerLevel->GetTimeStamp(*ResolvedPath.GetNonSandboxPath());
+}
+
+void FSandboxPlatformFile::SetTimeStamp(const TCHAR* Filename, FDateTime DateTime)
+{
+	const FSandboxedPlatformFilePath ResolvedPath = ToSandboxPath(Filename);
+
+	if (ResolvedPath.HasSandboxPath())
+	{
+		if (IsPathDeleted(ResolvedPath))
+		{
+			return;
+		}
+
+		if (LowerLevel->FileExists(*ResolvedPath.GetSandboxPath()))
+		{
+			return LowerLevel->SetTimeStamp(*ResolvedPath.GetSandboxPath(), DateTime);
+		}
+
+		return; // Do not attempt to modify the non-sandbox file
+	}
+
+	return LowerLevel->SetTimeStamp(*ResolvedPath.GetNonSandboxPath(), DateTime);
+}
+
+FDateTime FSandboxPlatformFile::GetAccessTimeStamp(const TCHAR* Filename)
+{
+	const FSandboxedPlatformFilePath ResolvedPath = ToSandboxPath(Filename);
+
+	if (ResolvedPath.HasSandboxPath())
+	{
+		if (IsPathDeleted(ResolvedPath))
+		{
+			return FDateTime::MinValue();
+		}
+
+		if (LowerLevel->FileExists(*ResolvedPath.GetSandboxPath()))
+		{
+			return LowerLevel->GetAccessTimeStamp(*ResolvedPath.GetSandboxPath());
+		}
+	}
+
+	return LowerLevel->GetAccessTimeStamp(*ResolvedPath.GetNonSandboxPath());
+}
+
+FString FSandboxPlatformFile::GetFilenameOnDisk(const TCHAR* Filename)
+{
+	const FSandboxedPlatformFilePath ResolvedPath = ToSandboxPath(Filename);
+	return LowerLevel->GetFilenameOnDisk(*ResolvedPath.GetNonSandboxPath());
+}
+
+IFileHandle* FSandboxPlatformFile::OpenRead(const TCHAR* Filename, bool bAllowWrite)
+{
+	const FSandboxedPlatformFilePath ResolvedPath = ToSandboxPath(Filename);
+
+	if (ResolvedPath.HasSandboxPath())
+	{
+		if (IsPathDeleted(ResolvedPath))
+		{
+			return nullptr;
+		}
+
+		if (bAllowWrite)
+		{
+			MigrateFileToSandbox(ResolvedPath);
+		}
+
+		if (bAllowWrite || LowerLevel->FileExists(*ResolvedPath.GetSandboxPath()))
+		{
+			return LowerLevel->OpenRead(*ResolvedPath.GetSandboxPath(), bAllowWrite);
+		}
+	}
+
+	return LowerLevel->OpenRead(*ResolvedPath.GetNonSandboxPath(), bAllowWrite);
+}
+
+IFileHandle* FSandboxPlatformFile::OpenReadNoBuffering(const TCHAR* Filename, bool bAllowWrite)
+{
+	const FSandboxedPlatformFilePath ResolvedPath = ToSandboxPath(Filename);
+
+	if (ResolvedPath.HasSandboxPath())
+	{
+		if (IsPathDeleted(ResolvedPath))
+		{
+			return nullptr;
+		}
+
+		if (bAllowWrite)
+		{
+			MigrateFileToSandbox(ResolvedPath);
+		}
+
+		if (bAllowWrite || LowerLevel->FileExists(*ResolvedPath.GetSandboxPath()))
+		{
+			return LowerLevel->OpenReadNoBuffering(*ResolvedPath.GetSandboxPath(), bAllowWrite);
+		}
+	}
+
+	return LowerLevel->OpenReadNoBuffering(*ResolvedPath.GetNonSandboxPath(), bAllowWrite);
+}
+
+IFileHandle* FSandboxPlatformFile::OpenWrite(const TCHAR* Filename, bool bAppend, bool bAllowRead)
+{
+	const FSandboxedPlatformFilePath ResolvedPath = ToSandboxPath(Filename);
+
+	if (ResolvedPath.HasSandboxPath())
+	{
+		MigrateFileToSandbox(ResolvedPath);
+
+		IFileHandle* Handle = LowerLevel->OpenWrite(*ResolvedPath.GetSandboxPath(), bAppend, bAllowRead);
+		if (Handle)
+		{
+			SetPathDeleted(ResolvedPath, false);
+		}
+		return Handle;
+	}
+
+	return LowerLevel->OpenWrite(*ResolvedPath.GetNonSandboxPath(), bAppend, bAllowRead);
+}
+
+bool FSandboxPlatformFile::DirectoryExists(const TCHAR* Directory)
+{
+	const FSandboxedPlatformFilePath ResolvedPath = ToSandboxPath(Directory);
+
+	if (ResolvedPath.HasSandboxPath())
+	{
+		if (IsPathDeleted(ResolvedPath))
+		{
+			return false;
+		}
+
+		if (LowerLevel->DirectoryExists(*ResolvedPath.GetSandboxPath()))
+		{
+			return true;
+		}
+	}
+
+	return LowerLevel->DirectoryExists(*ResolvedPath.GetNonSandboxPath());
+}
+
+bool FSandboxPlatformFile::CreateDirectory(const TCHAR* Directory)
+{
+	const FSandboxedPlatformFilePath ResolvedPath = ToSandboxPath(Directory);
+
+	if (ResolvedPath.HasSandboxPath())
+	{
+		if (LowerLevel->CreateDirectory(*ResolvedPath.GetSandboxPath()))
+		{
+			SetPathDeleted(ResolvedPath, false);
+			return true;
+		}
+
+		return false; // Do not attempt to create the non-sandbox directory
+	}
+
+	return LowerLevel->CreateDirectory(*ResolvedPath.GetNonSandboxPath());
+}
+
+bool FSandboxPlatformFile::DeleteDirectory(const TCHAR* Directory)
+{
+	const FSandboxedPlatformFilePath ResolvedPath = ToSandboxPath(Directory);
+
+	if (ResolvedPath.HasSandboxPath())
+	{
+		if (IsPathDeleted(ResolvedPath))
+		{
+			return true;
+		}
+
+		// Iterate the directory to validate that it is really empty before deleting it
+		const TArray<FDirectoryItem> DirectoryItems = GetDirectoryContents(ResolvedPath, Directory);
+		if (DirectoryItems.Num() == 0 && (!LowerLevel->DirectoryExists(*ResolvedPath.GetSandboxPath()) || LowerLevel->DeleteDirectory(*ResolvedPath.GetSandboxPath())))
+		{
+			SetPathDeleted(ResolvedPath, true);
+			return true;
+		}
+
+		return false; // Do not attempt to create the non-sandbox directory
+	}
+
+	return LowerLevel->DeleteDirectory(*ResolvedPath.GetNonSandboxPath());
+}
+
+FFileStatData FSandboxPlatformFile::GetStatData(const TCHAR* FilenameOrDirectory)
+{
+	const FSandboxedPlatformFilePath ResolvedPath = ToSandboxPath(FilenameOrDirectory);
+
+	if (ResolvedPath.HasSandboxPath())
+	{
+		if (IsPathDeleted(ResolvedPath))
+		{
+			return FFileStatData();
+		}
+
+		if (LowerLevel->FileExists(*ResolvedPath.GetSandboxPath()))
+		{
+			return LowerLevel->GetStatData(*ResolvedPath.GetSandboxPath());
+		}
+
+		FFileStatData StatData = LowerLevel->GetStatData(*ResolvedPath.GetNonSandboxPath());
+		// Sandbox files are always writeable.
+		StatData.bIsReadOnly = false;
+		return StatData;
+	}
+
+	return LowerLevel->GetStatData(*ResolvedPath.GetNonSandboxPath());
+}
+
+bool FSandboxPlatformFile::IterateDirectory(const TCHAR* Directory, FDirectoryVisitor& Visitor)
+{
+	const FSandboxedPlatformFilePath ResolvedPath = ToSandboxPath(Directory);
+
+	if (ResolvedPath.HasSandboxPath())
+	{
+		if (IsPathDeleted(ResolvedPath))
+		{
+			return false;
+		}
+
+		const TArray<FDirectoryItem> DirectoryItems = GetDirectoryContents(ResolvedPath, Directory);
+		for (const FDirectoryItem& DirectoryItem : DirectoryItems)
+		{
+			if (!Visitor.CallShouldVisitAndVisit(*DirectoryItem.Path, DirectoryItem.StatData.bIsDirectory))
+			{
+				return false;
+			}
+		}
+
+		return true;
+	}
+
+	return LowerLevel->IterateDirectory(Directory, Visitor); // Note: Using the path we were given here to ensure the calling code gets the expected path
+}
+
+bool FSandboxPlatformFile::IterateDirectoryStat(const TCHAR* Directory, FDirectoryStatVisitor& Visitor)
+{
+	const FSandboxedPlatformFilePath ResolvedPath = ToSandboxPath(Directory);
+
+	if (ResolvedPath.HasSandboxPath())
+	{
+		if (IsPathDeleted(ResolvedPath))
+		{
+			return false;
+		}
+
+		const TArray<FDirectoryItem> DirectoryItems = GetDirectoryContents(ResolvedPath, Directory);
+		for (const FDirectoryItem& DirectoryItem : DirectoryItems)
+		{
+			if (!Visitor.CallShouldVisitAndVisit(*DirectoryItem.Path, DirectoryItem.StatData))
+			{
+				return false;
+			}
+		}
+
+		return true;
+	}
+
+	return LowerLevel->IterateDirectoryStat(Directory, Visitor); // Note: Using the path we were given here to ensure the calling code gets the expected path
+}
+
+IAsyncReadFileHandle* FSandboxPlatformFile::OpenAsyncRead(const TCHAR* Filename, bool bAllowWrite /*= false*/)
+{
+	const FSandboxedPlatformFilePath ResolvedPath = ToSandboxPath(Filename);
+
+	if (ResolvedPath.HasSandboxPath())
+	{
+		if (IsPathDeleted(ResolvedPath) || LowerLevel->FileExists(*ResolvedPath.GetSandboxPath()))
+		{
+			return LowerLevel->OpenAsyncRead(*ResolvedPath.GetSandboxPath(), bAllowWrite);
+		}
+	}
+
+	return LowerLevel->OpenAsyncRead(*ResolvedPath.GetNonSandboxPath(), bAllowWrite);
+}
+
+void FSandboxPlatformFile::SetAsyncMinimumPriority(EAsyncIOPriorityAndFlags MinPriority)
+{
+	LowerLevel->SetAsyncMinimumPriority(MinPriority);
+}
+
+FString FSandboxPlatformFile::ConvertToAbsolutePathForExternalAppForRead(const TCHAR* Filename)
+{
+	const FSandboxedPlatformFilePath ResolvedPath = ToSandboxPath(Filename);
+
+	if (ResolvedPath.HasSandboxPath())
+	{
+		if (IsPathDeleted(ResolvedPath) || LowerLevel->FileExists(*ResolvedPath.GetSandboxPath()))
+		{
+			return LowerLevel->ConvertToAbsolutePathForExternalAppForRead(*ResolvedPath.GetSandboxPath());
+		}
+	}
+
+	return LowerLevel->ConvertToAbsolutePathForExternalAppForRead(*ResolvedPath.GetNonSandboxPath());
+}
+
+FString FSandboxPlatformFile::ConvertToAbsolutePathForExternalAppForWrite(const TCHAR* Filename)
+{
+	const FSandboxedPlatformFilePath ResolvedPath = ToSandboxPath(Filename);
+
+	if (ResolvedPath.HasSandboxPath())
+	{
+		return LowerLevel->ConvertToAbsolutePathForExternalAppForWrite(*ResolvedPath.GetSandboxPath());
+	}
+
+	return LowerLevel->ConvertToAbsolutePathForExternalAppForWrite(*ResolvedPath.GetNonSandboxPath());
+}
+
+
+bool FSandboxPlatformFile::CopyFileWithSCC(
+	const TCHAR* InToFilename,
+	const TCHAR* InFromFilename,
+	const FPersistArgs& InParam,
+	FPersistResult& OutResult
+	)
+{
+	// If this file maps to a package then we need to flush its linker so that we can overwrite the file on disk
+	FlushPackageFile(InToFilename);
+
+	ISourceControlProvider* SourceControlProvider = &ISourceControlModule::Get().GetProvider();
+	// Get the source control state of the destination file
+	FSourceControlStatePtr ToFileSCCState;
+	if (SourceControlProvider && SourceControlProvider->IsEnabled())
+	{
+		ToFileSCCState = SourceControlProvider->GetState(InToFilename, EStateCacheUsage::ForceUpdate);
+	}
+
+	// We don't need to do anything with source control if the file is already checked-out or added
+	const bool bRequiresSCCAction = ToFileSCCState && !ToFileSCCState->IsCheckedOut() && !ToFileSCCState->IsAdded();
+
+	// If the file can be checked-out, do so now
+	if (bRequiresSCCAction && ToFileSCCState->IsSourceControlled())
+	{
+		// the static analysis tool is not able to see that `SourceControlProvider` won't be null if `bRequiresSCCAction` is true
+		CA_ASSUME(SourceControlProvider != nullptr);
+		if (ToFileSCCState->CanCheckout() && SourceControlProvider->UsesCheckout())
+		{
+			TArray<FString> FilesToBeCheckedOut;
+			FilesToBeCheckedOut.Add(InToFilename);
+
+			if (SourceControlProvider->Execute(ISourceControlOperation::Create<FCheckOut>(), FilesToBeCheckedOut) != ECommandResult::Succeeded)
+			{
+				FText Failure = FText::Format(
+					NSLOCTEXT("Concert", "CheckoutFileFailure", "Failed to check-out file '{0}' from revision control when persiting sandbox state!"),
+					FText::FromString(InToFilename));
+				UE_LOG(LogFileSandbox, Warning, TEXT("%s"),  *Failure.ToString());
+				OutResult.FailureReasons.Add(MoveTemp(Failure));
+				return false;
+			}
+		}
+		else
+		{
+			FText Failure = FText::Format(
+				NSLOCTEXT("Concert", "CanCheckoutFileFailure", "Can't check-out file '{0}' from revision control when persiting sandbox state!"),
+				FText::FromString(InToFilename));
+			UE_LOG(LogFileSandbox, Warning, TEXT("%s"), *Failure.ToString());
+			OutResult.FailureReasons.Add(MoveTemp(Failure));
+			return false;
+		}
+	}
+	else if (InParam.bShouldMakeWritableIfNoSourceControl)
+	{
+		if (LowerLevel->FileExists(InToFilename) && !LowerLevel->SetReadOnly(InToFilename, false))
+		{
+			FText Failure = FText::Format(
+				NSLOCTEXT("Concert", "CanMakeWritable", "Can't make file '{0}' writable when persiting sandbox state!"),
+				FText::FromString(InToFilename));
+			UE_LOG(LogFileSandbox, Warning, TEXT("%s"), *Failure.ToString());
+			OutResult.FailureReasons.Add(MoveTemp(Failure));
+			return false;
+		}
+	}
+
+	// Copy the on-disk sandbox file
+	FString ToFileDir = FPaths::GetPath(InToFilename);
+	if (!CreateDirectoryTree(*ToFileDir) || !LowerLevel->CopyFile(InToFilename, InFromFilename, EPlatformFileRead::None, EPlatformFileWrite::AttemptDeleteAndCreate))
+	{
+		FText Failure = FText::Format(
+			NSLOCTEXT("Concert", "CopyFileFailure", "Failed to copy file '{0}' (from '{1}') when persiting sandbox state!"),
+			FText::FromString(InToFilename), FText::FromString(InFromFilename));
+		UE_LOG(LogFileSandbox, Warning, TEXT("%s"), *Failure.ToString());
+		OutResult.FailureReasons.Add(MoveTemp(Failure));
+		return false;
+	}
+
+		// If the file is new, add it to source control now
+	if (bRequiresSCCAction && !ToFileSCCState->IsSourceControlled())
+	{
+		TArray<FString> FilesToBeAdded;
+		FilesToBeAdded.Add(InToFilename);
+
+		if (SourceControlProvider->Execute(ISourceControlOperation::Create<FMarkForAdd>(), FilesToBeAdded) != ECommandResult::Succeeded)
+		{
+			FText Failure = FText::Format(
+				NSLOCTEXT("Concert", "AddFileFailure", "Failed to add file '{0}' to revision control when persiting sandbox state!"),
+				FText::FromString(InToFilename));
+			UE_LOG(LogFileSandbox, Warning, TEXT("%s"), *Failure.ToString());
+			OutResult.FailureReasons.Add(MoveTemp(Failure));
+			return false;
+		}
+	}
+	return true;
+}
+
+/** Delete file using source control */
+bool FSandboxPlatformFile::DeleteFileWithSCC(const TCHAR* InFilename, FPersistResult& OutResult)
+{
+	// If this file maps to a package then we need to flush its linker so that we can remove the file from disk
+	FlushPackageFile(InFilename);
+
+	ISourceControlProvider* SourceControlProvider = &ISourceControlModule::Get().GetProvider();
+	// Get the source control state of the file
+	FSourceControlStatePtr FileSCCState;
+	if (SourceControlProvider && SourceControlProvider->IsEnabled())
+	{
+		FileSCCState = SourceControlProvider->GetState(InFilename, EStateCacheUsage::ForceUpdate);
+	}
+
+	// Try and let source control remove the file first
+	if (FileSCCState && FileSCCState->IsSourceControlled())
+	{
+		const bool bAdded = FileSCCState->IsAdded();
+
+		if (bAdded || FileSCCState->IsCheckedOut())
+		{
+			if (SourceControlProvider->Execute(ISourceControlOperation::Create<FRevert>(), InFilename) != ECommandResult::Succeeded)
+			{
+				FText Failure = FText::Format(NSLOCTEXT("Concert", "RevertFileFailure", "Failed to revert file '{0}' to revision control when persisting sandbox state!"), FText::FromString(InFilename));
+				UE_LOG(LogFileSandbox, Warning, TEXT("%s"), *Failure.ToString());
+				OutResult.FailureReasons.Add(MoveTemp(Failure));
+				return false;
+			}
+
+			if (!bAdded)
+			{
+				if (SourceControlProvider->Execute(ISourceControlOperation::Create<FDelete>(), InFilename) != ECommandResult::Succeeded)
+				{
+					FText Failure = FText::Format(NSLOCTEXT("Concert", "DeleteSCCFileFailure", "Failed to delete file '{0}' from revision control when persiting sandbox state!"), FText::FromString(InFilename));
+					UE_LOG(LogFileSandbox, Warning, TEXT("%s"), *Failure.ToString());
+					OutResult.FailureReasons.Add(MoveTemp(Failure));
+					return false;
+				}
+			}
+		}
+	}
+
+	// Delete file if it still exists
+	if (LowerLevel->FileExists(InFilename) && !LowerLevel->DeleteFile(InFilename))
+	{
+		FText Failure = FText::Format(NSLOCTEXT("Concert", "DeleteFileFailure", "Failed to delete file '{0}' when persiting sandbox state!"), FText::FromString(InFilename));
+		UE_LOG(LogFileSandbox, Warning, TEXT("%s"), *Failure.ToString());
+		OutResult.FailureReasons.Add(MoveTemp(Failure));
+		return false;
+	}
+	return true;
+}
+
+FPersistResult FSandboxPlatformFile::PersistSandbox(
+	FPersistArgs InParams, TFunctionRef<void(const FSandboxedPlatformFilePath&, EPersistedFileType)> InOnPersistFile
+	)
+{
+	FPersistResult OutResult;
+	// We need to disable the sandbox while we do this
+	bool bFullOperationSuccess = true;
+	TGuardValue<TAtomic<bool>, bool> DisableSandboxGuard(bSandboxEnabled, false);
+	for (const FString& File : InParams.Files)
+	{
+		const FSandboxedPlatformFilePath FilePath = ToSandboxPath(File, true);
+		const EPersistedFileType Type = IsPathDeleted(FilePath) ? EPersistedFileType::Deleted : EPersistedFileType::Moved; 
+		const bool bSuccess = Type == EPersistedFileType::Deleted
+			? DeleteFileWithSCC(*FilePath.GetNonSandboxPath(), OutResult)
+			: CopyFileWithSCC(*FilePath.GetNonSandboxPath(), *FilePath.GetSandboxPath(), InParams, OutResult);
+
+		// if the file operation was successful, mark the file as persisted
+		if (bSuccess)
+		{
+			InOnPersistFile(FilePath, Type);
+		}
+
+		bFullOperationSuccess &= bSuccess;
+	}
+	OutResult.PersistStatus = bFullOperationSuccess ? EPersistStatus::Success : EPersistStatus::Failure;
+	return OutResult;
+}
+
+void FSandboxPlatformFile::DiscardSandbox(
+	TArray<FName>& OutPackagesPendingHotReload, TArray<FName>& OutPackagesPendingPurge, TArray<FFileChangeData>* OutFileChanges
+	)
+{
+	// We need to disable the sandbox while we do this
+	TGuardValue<TAtomic<bool>, bool> DisableSandboxGuard(bSandboxEnabled, false);
+
+	TArray<FFileChangeData> FileChanges;
+
+	// Add any files that were deleted by the sandbox but exist in the non-sandbox directory
+	{
+		FScopeLock DeletedSandboxPathsLock(&DeletedSandboxPathsCS);
+		for (auto It = DeletedSandboxPaths.CreateIterator(); It; ++It)
+		{
+			// If this file maps to a package then we need to flush its linker so that we can remove the file from the sandbox
+			FName PackageName;
+			FlushPackageFile(It->GetNonSandboxPath(), &PackageName);
+
+			if (LowerLevel->FileExists(*It->GetNonSandboxPath()))
+			{
+				if (!PackageName.IsNone())
+				{
+					OutPackagesPendingHotReload.Add(PackageName);
+				}
+				FileChanges.Emplace(It->GetNonSandboxPath(), FFileChangeData::FCA_Added);
+			}
+			It.RemoveCurrent();
+			continue;
+		}
+
+		// Clear the deleted path information
+		DeletedSandboxPaths.Reset();
+	}
+
+	{
+		FScopeLock SandboxMountPointsLock(&SandboxMountPointsCS);
+		for (const FSandboxMountPoint& SandboxMountPoint : SandboxMountPoints)
+		{
+			// Modify any files that exist in both the non-sandbox and sandbox directories
+			// Delete any files that exist in the sandbox but don't exist in the non-sandbox directory
+			LowerLevel->IterateDirectoryRecursively(*SandboxMountPoint.Path.GetSandboxPath(), [this, &SandboxMountPoint, &FileChanges, &OutPackagesPendingHotReload, &OutPackagesPendingPurge](const TCHAR* InFilenameOrDirectory, const bool InIsDirectory) -> bool
+			{
+				if (!InIsDirectory)
+				{
+					const FSandboxedPlatformFilePath RemappedFilePath = FSandboxedPlatformFilePath::CreateNonSandboxPath(FPaths::ConvertRelativePathToFull(InFilenameOrDirectory), SandboxMountPoint.Path);
+
+					// If this file maps to a package then we need to flush its linker so that we can remove the file from the sandbox
+					FName PackageName;
+					FlushPackageFile(RemappedFilePath.GetNonSandboxPath(), &PackageName, false);
+
+					if (LowerLevel->FileExists(*RemappedFilePath.GetNonSandboxPath()))
+					{
+						if (!PackageName.IsNone())
+						{
+							OutPackagesPendingHotReload.Add(PackageName);
+							OutPackagesPendingPurge.Remove(PackageName);
+						}
+						
+						FileChanges.Emplace(RemappedFilePath.GetNonSandboxPath(), FFileChangeData::FCA_Modified);
+					}
+					else
+					{
+						if (!PackageName.IsNone())
+						{
+							OutPackagesPendingPurge.Add(PackageName);
+							OutPackagesPendingHotReload.Remove(PackageName);
+						}
+						
+						FileChanges.Emplace(RemappedFilePath.GetNonSandboxPath(), FFileChangeData::FCA_Removed);
+					}
+				}
+				return true; // Continue iteration
+			});
+
+			//FString SandboxRootPath = GetSandboxRootPath();
+			//FString Stashed = SandboxRootPath.Replace(TEXT("Sandbox"), TEXT("Stashed"));
+			//LowerLevel->MoveFile(SandboxRootPath, Stashed);
+			// Delete everything under the mount point
+			//LowerLevel->IterateDirectory(*SandboxMountPoint.Path.GetSandboxPath(), [this](const TCHAR* InFilenameOrDirectory, const bool InIsDirectory) -> bool
+			//{
+			//	if (InIsDirectory)
+			//	{
+			//		LowerLevel->DeleteDirectoryRecursively(InFilenameOrDirectory);
+			//	}
+			//	else
+			//	{
+			//		LowerLevel->DeleteFile(InFilenameOrDirectory);
+			//	}
+			//	return true; // Continue iteration
+			//});
+		}
+	}
+
+#if WITH_EDITOR
+	// Notify that the sandboxed directories have been restored to their original state
+	if (FDirectoryWatcherModule* DirectoryWatcherModule = GetDirectoryWatcherModuleIfLoaded())
+	{
+		if (FileChanges.Num() > 0)
+		{
+			DirectoryWatcherModule->RegisterExternalChanges(FileChanges);
+
+			// Force a sync of the Asset Registry immediately to process the
+			// file changes. This ensures that the Asset Registry is brought
+			// up to date especially after file/package deletes, since
+			// otherwise it may not get updated before packages are hot
+			// reloaded and stale entries could still be present in the
+			// registry.
+			SynchronizeAssetRegistry();
+		}
+	}
+#endif
+
+	if (OutFileChanges)
+	{
+		*OutFileChanges = MoveTemp(FileChanges);
+	}
+}
+
+TArray<FString> FSandboxPlatformFile::GatherSandboxChangedFilenames(TFunctionRef<bool(const FSandboxedPlatformFilePath& InFile)> InShouldSkipFile) const
+{
+	TArray<FString> ChangedFiles;
+	
+	// Gather deleted paths
+	{
+		FScopeLock DeletedSandboxPathsLock(&DeletedSandboxPathsCS);
+		ChangedFiles.Reserve(DeletedSandboxPaths.Num());
+		Algo::TransformIf(DeletedSandboxPaths, ChangedFiles,
+			[this, &InShouldSkipFile](const FSandboxedPlatformFilePath& DeletedPath)
+			{
+				return !InShouldSkipFile(DeletedPath) && LowerLevel->FileExists(*DeletedPath.GetNonSandboxPath());
+			},
+			[this](const FSandboxedPlatformFilePath& DeletedPath)
+			{
+				return DeletedPath.GetNonSandboxPath();
+			});
+	}
+
+	// Gather mounts files
+	{
+		FScopeLock SandboxMountPointsLock(&SandboxMountPointsCS);
+		for (const FSandboxMountPoint& MountPoint : SandboxMountPoints)
+		{
+			LowerLevel->IterateDirectoryRecursively(*MountPoint.Path.GetSandboxPath(), [this, &InShouldSkipFile, &MountPoint, &ChangedFiles](const TCHAR* InFilenameOrDirectory, const bool InIsDirectory) -> bool
+			{
+				if (InIsDirectory)
+				{
+					return true;
+				}
+
+				const FSandboxedPlatformFilePath Path = FSandboxedPlatformFilePath::CreateNonSandboxPath(
+					FPaths::ConvertRelativePathToFull(InFilenameOrDirectory), MountPoint.Path
+					);
+				if (!InShouldSkipFile(Path))
+				{
+					ChangedFiles.Add(Path.GetNonSandboxPath());
+				}
+					
+				return true; // continue iteration
+			});
+		}
+	}
+
+	return ChangedFiles;
+}
+
+bool FSandboxPlatformFile::DeletedPackageExistsInNonSandbox(FString InFilename) const
+{
+	const FSandboxedPlatformFilePath ResolvedPath = ToSandboxPath(MoveTemp(InFilename));
+	return ResolvedPath.HasSandboxPath()
+		&& IsPathDeleted(ResolvedPath)
+		&& LowerLevel->FileExists(*ResolvedPath.GetNonSandboxPath());
+}
+
+void FSandboxPlatformFile::OnContentPathMounted(const FString& InAssetPath, const FString& InFilesystemPath)
+{
+	RegisterContentMountPath(InAssetPath, InFilesystemPath);
+}
+
+void FSandboxPlatformFile::OnContentPathDismounted(const FString& InAssetPath, const FString& InFilesystemPath)
+{
+	UnregisterContentMountPath(InAssetPath, InFilesystemPath);
+}
+
+void FSandboxPlatformFile::RegisterContentMountPath(const FString& InAssetPath, const FString& InFilesystemPath)
+{
+	FString AbsoluteSandboxPath = FPaths::ConvertRelativePathToFull(GetSandboxRootPath() / InAssetPath) / TEXT("");
+	FString AbsoluteNonSandboxPath = FPaths::ConvertRelativePathToFull(InFilesystemPath) / TEXT("");
+	const bool bCreatedDirectory = LowerLevel->CreateDirectory(*AbsoluteSandboxPath);
+	
+	UE_CLOG(!bCreatedDirectory, LogFileSandbox, Error,
+		TEXT("Failed to create sandboxed mount directory %s for mount point %s"),
+		*AbsoluteSandboxPath, *AbsoluteNonSandboxPath
+		);
+	
+	if (bCreatedDirectory)
+	{
+		const FScopeLock SandboxMountPointsLock(&SandboxMountPointsCS);
+
+		FSandboxMountPoint& SandboxMountPoint = SandboxMountPoints.Add_GetRef(FSandboxMountPoint{ FSandboxedPlatformFilePath(MoveTemp(AbsoluteNonSandboxPath), MoveTemp(AbsoluteSandboxPath)) });
+#if WITH_EDITOR
+		if (IDirectoryWatcher* DirectoryWatcher = GetDirectoryWatcher())
+		{
+			DirectoryWatcher->RegisterDirectoryChangedCallback_Handle(SandboxMountPoint.Path.GetSandboxPath(), IDirectoryWatcher::FDirectoryChanged::CreateRaw(this, &FSandboxPlatformFile::OnDirectoryChanged, SandboxMountPoint.Path), SandboxMountPoint.OnDirectoryChangedHandle, IDirectoryWatcher::IncludeDirectoryChanges);
+		}
+#endif
+	}
+}
+
+void FSandboxPlatformFile::UnregisterContentMountPath(const FString& InAssetPath, const FString& InFilesystemPath)
+{
+	const FString AbsoluteSandboxPath = FPaths::ConvertRelativePathToFull(GetSandboxRootPath() / InAssetPath) / TEXT("");
+	{
+		const FScopeLock SandboxMountPointsLock(&SandboxMountPointsCS);
+		SandboxMountPoints.RemoveAll([&AbsoluteSandboxPath](FSandboxMountPoint& InSandboxMountPoint) -> bool
+		{
+			const bool bShouldRemove = InSandboxMountPoint.Path.GetSandboxPath() == AbsoluteSandboxPath;
+#if WITH_EDITOR
+			if (bShouldRemove && InSandboxMountPoint.OnDirectoryChangedHandle.IsValid())
+			{
+				if (IDirectoryWatcher* DirectoryWatcher = GetDirectoryWatcherIfLoaded())
+				{
+					DirectoryWatcher->UnregisterDirectoryChangedCallback_Handle(InSandboxMountPoint.Path.GetSandboxPath(), InSandboxMountPoint.OnDirectoryChangedHandle);
+					InSandboxMountPoint.OnDirectoryChangedHandle.Reset();
+				}
+			}
+#endif
+			return bShouldRemove;
+		});
+	}
+	{
+		const FScopeLock DeletedSandboxPathsLock(&DeletedSandboxPathsCS);
+		for (auto It = DeletedSandboxPaths.CreateIterator(); It; ++It)
+		{
+			if (It->GetSandboxPath().StartsWith(AbsoluteSandboxPath))
+			{
+				It.RemoveCurrent();
+				continue;
+			}
+		}
+	}
+	LowerLevel->DeleteDirectoryRecursively(*AbsoluteSandboxPath);
+}
+
+FSandboxedPlatformFilePath FSandboxPlatformFile::ToSandboxPath(FString InFilename, const bool bEvenIfDisabled) const
+{
+	return ToSandboxPath_Absolute(FPaths::ConvertRelativePathToFull(MoveTemp(InFilename)), bEvenIfDisabled);
+}
+
+FSandboxedPlatformFilePath FSandboxPlatformFile::ToSandboxPath_Absolute(FString InFilename, const bool bEvenIfDisabled) const
+{
+	if (bEvenIfDisabled || IsSandboxEnabled())
+	{
+		FScopeLock SandboxMountPointsLock(&SandboxMountPointsCS);
+		for (const FSandboxMountPoint& SandboxMountPoint : SandboxMountPoints)
+		{
+			// Mount point are stored with a trailing slash to prevent matching mount point with similar names -> (/Bla/Content, /Bla/ContentSupreme)
+			// So we test without the slash to make sure with can match mount point directly -> (/Bla/Content match /Bla/Content/)
+			const FString& PathStr = SandboxMountPoint.Path.GetNonSandboxPath();
+			int32 PathStrNoSlashLength = PathStr.Len() - 1;
+			if (FCString::Strnicmp(*InFilename, *PathStr, PathStrNoSlashLength) == 0 && (InFilename.Len() == PathStrNoSlashLength || InFilename[PathStrNoSlashLength] == TEXT('/')))
+			{
+				return FSandboxedPlatformFilePath::CreateSandboxPath(MoveTemp(InFilename), SandboxMountPoint.Path);
+			}
+		}
+	}
+
+	return FSandboxedPlatformFilePath(MoveTemp(InFilename));
+}
+
+FSandboxedPlatformFilePath FSandboxPlatformFile::FromSandboxPath(FString InFilename) const
+{
+	return FromSandboxPath_Absolute(FPaths::ConvertRelativePathToFull(MoveTemp(InFilename)));
+}
+
+FSandboxedPlatformFilePath FSandboxPlatformFile::FromSandboxPath_Absolute(FString InFilename) const
+{
+	FScopeLock SandboxMountPointsLock(&SandboxMountPointsCS);
+
+	for (const FSandboxMountPoint& SandboxMountPoint : SandboxMountPoints)
+	{
+		// Mount point are stored with a trailing slash to prevent matching mount point with similar names -> (/Bla/Content, /Bla/ContentSupreme)
+		// So we test without the slash to make sure with can match mount point directly -> (/Bla/Content match /Bla/Content/)
+		const FString& PathStr = SandboxMountPoint.Path.GetSandboxPath();
+		int32 PathStrNoSlashLength = PathStr.Len() - 1;
+		if (FCString::Strnicmp(*InFilename, *PathStr, PathStrNoSlashLength) == 0 && (InFilename.Len() == PathStrNoSlashLength || InFilename[PathStrNoSlashLength] == TEXT('/')))
+		{
+			return FSandboxedPlatformFilePath::CreateNonSandboxPath(MoveTemp(InFilename), SandboxMountPoint.Path);
+		}
+	}
+
+	return FSandboxedPlatformFilePath(MoveTemp(InFilename));
+}
+
+bool FSandboxPlatformFile::IsPathDeleted(const FSandboxedPlatformFilePath& InPath) const
+{
+	FScopeLock DeletedSandboxPathsLock(&DeletedSandboxPathsCS);
+	return DeletedSandboxPaths.Contains(InPath);
+}
+
+void FSandboxPlatformFile::SetPathDeleted(const FSandboxedPlatformFilePath& InPath, const bool bIsDeleted)
+{
+	FScopeLock DeletedSandboxPathsLock(&DeletedSandboxPathsCS);
+	if (bIsDeleted)
+	{
+		DeletedSandboxPaths.Add(InPath);
+	}
+	else
+	{
+		DeletedSandboxPaths.Remove(InPath);
+	}
+	
+	OnFileDeletedDelegate.Broadcast(InPath, bIsDeleted);
+}
+
+void FSandboxPlatformFile::NotifyFileDeleted(const FSandboxedPlatformFilePath& InPath)
+{
+	if (!IsSandboxEnabled())
+	{
+		return;
+	}
+
+#if WITH_EDITOR
+	if (FDirectoryWatcherModule* DirectoryWatcherModule = GetDirectoryWatcherModuleIfLoaded())
+	{
+		FFileChangeData FileChange(InPath.GetNonSandboxPath(), FFileChangeData::FCA_Removed);
+		DirectoryWatcherModule->RegisterExternalChanges(TArrayView<const FFileChangeData>(&FileChange, 1));
+	}
+#endif
+}
+
+void FSandboxPlatformFile::MigrateFileToSandbox(const FSandboxedPlatformFilePath& InPath) const
+{
+	checkf(InPath.HasSandboxPath(), TEXT("MigrateFileToSandbox requires a sandbox path to be set!"));
+
+	// Migrate the non-sandbox directory structure to the sandbox
+	{
+		const FString SandboxDirectoryPath = FPaths::GetPath(InPath.GetSandboxPath());
+
+		// We create the directory if no part of it has been explicitly deleted in this sandbox
+		bool bCreateDirectory = true;
+		{
+			// Walk the paths backwards for as long as they match (which is the sandbox relative part of the paths)
+			FSandboxedPlatformFilePath TmpSandboxFilePath(FPaths::GetPath(InPath.GetNonSandboxPath()), CopyTemp(SandboxDirectoryPath));
+			while (FPaths::GetBaseFilename(TmpSandboxFilePath.GetNonSandboxPath()) == FPaths::GetBaseFilename(TmpSandboxFilePath.GetSandboxPath()))
+			{
+				if (IsPathDeleted(TmpSandboxFilePath))
+				{
+					bCreateDirectory = false;
+					break;
+				}
+				TmpSandboxFilePath = FSandboxedPlatformFilePath(FPaths::GetPath(TmpSandboxFilePath.GetNonSandboxPath()), FPaths::GetPath(TmpSandboxFilePath.GetSandboxPath()));
+			}
+		}
+		if (bCreateDirectory)
+		{
+			LowerLevel->CreateDirectoryTree(*SandboxDirectoryPath);
+		}
+	}
+
+	if (IsPathDeleted(InPath))
+	{
+		// Sandbox has explicitly deleted this file - don't resurrect it from the non-sandbox file
+		return;
+	}
+
+	if (LowerLevel->FileExists(*InPath.GetSandboxPath()))
+	{
+		UE_LOG(LogFileSandbox, Verbose, TEXT("MigrateFileToSandbox: Sandbox has the file %s already at this location. Nothing to do."), *InPath.GetSandboxPath() );
+		// Sandbox already has a file at this location - nothing to do
+		return;
+	}
+
+	if (!LowerLevel->FileExists(*InPath.GetNonSandboxPath()))
+	{
+		UE_LOG(LogFileSandbox, Verbose, TEXT("MigrateFileToSandbox: There is no file %s at this location. Nothing to do."), *InPath.GetNonSandboxPath() );
+		// Non-sandbox has no file at this location - nothing to do
+		return;
+	}
+
+	UE_LOG(LogFileSandbox, Verbose, TEXT("MigrateFileToSandbox: Copying file %s into the sandbox."), *InPath.GetNonSandboxPath() );
+	// Copy the file into the sandbox
+	LowerLevel->CopyFile(*InPath.GetSandboxPath(), *InPath.GetNonSandboxPath());
+
+	// Ensure the migrated file is writable
+	LowerLevel->SetReadOnly(*InPath.GetSandboxPath(), false);
+}
+
+TArray<FSandboxPlatformFile::FDirectoryItem> FSandboxPlatformFile::GetDirectoryContents(const FSandboxedPlatformFilePath& InPath, const TCHAR* InDirBase) const
+{
+	checkf(InPath.HasSandboxPath(), TEXT("GetDirectoryContents requires a sandbox path to be set!"));
+	TMap<FString, FFileStatData> FoundItems;
+
+	// Gather the items, the sandbox iteration is straightforward
+	LowerLevel->IterateDirectoryStat(*InPath.GetSandboxPath(), [&FoundItems](const TCHAR* FilenameOrDirectory, const FFileStatData& StatData)
+	{
+		FoundItems.Add(FPaths::GetCleanFilename(FilenameOrDirectory), StatData);
+		return true;
+	});
+	// Gather the non-sandbox, validating we haven't already gathered the sandbox equivalent or the file/dir is marked as deleted
+	LowerLevel->IterateDirectoryStat(*InPath.GetNonSandboxPath(), [this, &InPath, &FoundItems](const TCHAR* FilenameOrDirectory, const FFileStatData& StatData)
+	{
+		FString SandboxFilenameOrDirectory = FilenameOrDirectory;
+		FString CleanFilenameOrDir = FPaths::GetCleanFilename(SandboxFilenameOrDirectory);
+
+		if (!FoundItems.Contains(CleanFilenameOrDir) &&
+			!IsPathDeleted(FSandboxedPlatformFilePath::CreateSandboxPath(MoveTemp(SandboxFilenameOrDirectory), InPath)))
+		{
+			FoundItems.Add(MoveTemp(CleanFilenameOrDir), StatData);
+		}
+		return true;
+	});
+
+	// Turn the found items in an array and re-base on InDirBase
+	FString DirBase(InDirBase);
+	TArray<FDirectoryItem> DirectoryContents;
+	DirectoryContents.Reserve(FoundItems.Num());
+	for (const auto& FoundItemPair : FoundItems)
+	{
+		DirectoryContents.Add(FDirectoryItem{ DirBase / FoundItemPair.Key, FoundItemPair.Value });
+	}
+
+	// Sort the result
+	DirectoryContents.Sort([](const FDirectoryItem& InOne, const FDirectoryItem& InTwo) -> bool
+	{
+		return InOne.Path < InTwo.Path;
+	});
+
+	return DirectoryContents;
+}
+
+#if WITH_EDITOR
+
+void FSandboxPlatformFile::OnDirectoryChanged(const TArray<FFileChangeData>& FileChanges, FSandboxedPlatformFilePath MountPath)
+{
+	if (!IsSandboxEnabled())
+	{
+		return;
+	}
+
+	if (FDirectoryWatcherModule* DirectoryWatcherModule = GetDirectoryWatcherModuleIfLoaded())
+	{
+		TArray<FFileChangeData> RemappedFileChanges;
+		RemappedFileChanges.Reserve(FileChanges.Num());
+
+		// Map the sandbox paths back to their original roots and notify the directory watcher
+		for (const FFileChangeData& FileChange : FileChanges)
+		{
+			const FSandboxedPlatformFilePath RemappedFilePath = FSandboxedPlatformFilePath::CreateNonSandboxPath(FPaths::ConvertRelativePathToFull(FileChange.Filename), MountPath);
+			RemappedFileChanges.Emplace(RemappedFilePath.GetNonSandboxPath(), FileChange.Action);
+
+			// Make sure the deleted state of this item is synchronized correctly
+			if (FileChange.Action == FFileChangeData::FCA_Added)
+			{
+				SetPathDeleted(RemappedFilePath, false);
+			}
+			else if (FileChange.Action == FFileChangeData::FCA_Removed)
+			{
+				SetPathDeleted(RemappedFilePath, true);
+			}
+		}
+
+		DirectoryWatcherModule->RegisterExternalChanges(RemappedFileChanges);
+	}
+}
+
+#endif
+}
\ No newline at end of file
diff --git a/Engine/Plugins/Developer/Sandbox/FileSandbox/Source/FileSandboxCore/Private/Sandbox/Platform/SandboxPlatformFile.h b/Engine/Plugins/Developer/Sandbox/FileSandbox/Source/FileSandboxCore/Private/Sandbox/Platform/SandboxPlatformFile.h
new file mode 100644
--- /dev/null
+++ b/Engine/Plugins/Developer/Sandbox/FileSandbox/Source/FileSandboxCore/Private/Sandbox/Platform/SandboxPlatformFile.h
@@ -0,0 +1,180 @@
+﻿// Copyright Epic Games, Inc. All Rights Reserved.
+
+#pragma once
+
+#include "Containers/Map.h"
+#include "Containers/Set.h"
+#include "Delegates/IDelegateInstance.h"
+#include "GenericPlatform/GenericPlatformFile.h"
+#include "IDirectoryWatcher.h"
+#include "SandboxedPlatformFilePath.h"
+
+namespace UE::FileSandboxCore
+{
+struct FPersistResult;
+struct FPersistArgs;
+
+DECLARE_MULTICAST_DELEGATE_TwoParams(FDeletionStateChanged, const FSandboxedPlatformFilePath&, bool /* bIsDeleted */);
+
+class FSandboxPlatformFile : public IPlatformFile
+{
+public:
+
+	explicit FSandboxPlatformFile(const FString& InSandboxRootPath);
+	virtual ~FSandboxPlatformFile() override;
+	
+	//~ For visibility of overloads we don't override
+	using IPlatformFile::IterateDirectory;
+	using IPlatformFile::IterateDirectoryStat;
+	
+	//~ begin IPlatformFile overrides
+	virtual void SetSandboxEnabled(bool bInEnabled) override;
+	virtual bool IsSandboxEnabled() const override;
+	virtual bool Initialize(IPlatformFile* Inner, const TCHAR* CmdLine) override;
+	virtual void Tick() override;
+	virtual IPlatformFile* GetLowerLevel() override;
+	virtual void SetLowerLevel(IPlatformFile* NewLowerLevel) override;
+	virtual const TCHAR* GetName() const override;
+	virtual bool FileExists(const TCHAR* Filename) override;
+	virtual int64 FileSize(const TCHAR* Filename) override;
+	virtual bool DeleteFile(const TCHAR* Filename) override;
+	virtual bool IsReadOnly(const TCHAR* Filename) override;
+	virtual bool SetReadOnly(const TCHAR* Filename, bool bNewReadOnlyValue) override;
+	virtual bool MoveFile(const TCHAR* To, const TCHAR* From) override;
+	virtual FDateTime GetTimeStamp(const TCHAR* Filename) override;
+	virtual void SetTimeStamp(const TCHAR* Filename, FDateTime DateTime) override;
+	virtual FDateTime GetAccessTimeStamp(const TCHAR* Filename) override;
+	virtual FString GetFilenameOnDisk(const TCHAR* Filename) override;
+	virtual IFileHandle* OpenRead(const TCHAR* Filename, bool bAllowWrite = false) override;
+	virtual IFileHandle* OpenReadNoBuffering(const TCHAR* Filename, bool bAllowWrite = false) override;
+	virtual IFileHandle* OpenWrite(const TCHAR* Filename, bool bAppend = false, bool bAllowRead = false) override;
+	virtual bool DirectoryExists(const TCHAR* Directory) override;
+	virtual bool CreateDirectory(const TCHAR* Directory) override;
+	virtual bool DeleteDirectory(const TCHAR* Directory) override;
+	virtual FFileStatData GetStatData(const TCHAR* FilenameOrDirectory) override;
+	virtual bool IterateDirectory(const TCHAR* Directory, FDirectoryVisitor& Visitor) override;
+	virtual bool IterateDirectoryStat(const TCHAR* Directory, FDirectoryStatVisitor& Visitor) override;
+	virtual IAsyncReadFileHandle* OpenAsyncRead(const TCHAR* Filename, bool bAllowWrite = false) override;
+	virtual void SetAsyncMinimumPriority(EAsyncIOPriorityAndFlags MinPriority) override;
+	virtual FString ConvertToAbsolutePathForExternalAppForRead(const TCHAR* Filename) override;
+	virtual FString ConvertToAbsolutePathForExternalAppForWrite(const TCHAR* Filename) override;
+	//~ end IPlatformFile overrides
+
+	enum class EPersistedFileType : uint8 { Deleted, Moved };
+	/**
+	 *	Persist the file list from the sandbox state onto the real files Will mark files for which the operation was
+	 *	successful as persisted.
+	 *	@param InParams	Parameters controlling the persist options specified by the user.
+	 *	@param InOnPersistFile Invoked for each file successfully persisted.
+	 *	@return Returns a persist result object that contains result status.
+	 */
+	FPersistResult PersistSandbox(FPersistArgs InParams, TFunctionRef<void(const FSandboxedPlatformFilePath&, EPersistedFileType)> InOnPersistFile);
+
+	/** 
+	 *	Discard the sandbox state
+	 *	This will trigger directory watcher notifications for files that are restored.
+	 *	This will also gather packages that need to be hot reloaded or purge from memory
+	 *	@param OutPackagePendingHotReload	Array to be filled with package name that need to be hot reloaded
+	 *	@param OutPackagesPendingPurge		Array to be filled with package name that need to be purged
+	 *	@param OutFileChanges				Array to be filled with changes made to the filesystem
+	 */
+	void DiscardSandbox(TArray<FName>& OutPackagesPendingHotReload, TArray<FName>& OutPackagesPendingPurge, TArray<FFileChangeData>* OutFileChanges = nullptr);
+
+	/** 
+	 * Gather all files changes that are currently in the sandbox.
+	 * @return Array of filepaths in no particular order.
+	 */
+	TArray<FString> GatherSandboxChangedFilenames(TFunctionRef<bool(const FSandboxedPlatformFilePath& InFile)> InShouldSkipFile = [](auto){ return true; }) const;
+	
+	/** @return Whether the given package file was deleted and exists on non sandbox path. */
+	bool DeletedPackageExistsInNonSandbox(FString InFilename) const;
+
+	/** Invoked when a file is deleted. */
+	FDeletionStateChanged& OnDeletionStateChanged() { return OnFileDeletedDelegate; }
+	
+private:
+	
+	struct FDirectoryItem
+	{
+		FString Path;
+		FFileStatData StatData;
+	};
+
+	struct FSandboxMountPoint
+	{
+		/** Sandbox path */
+		FSandboxedPlatformFilePath Path;
+
+		/** Sandbox directory watcher delegate handle (if any) */
+		FDelegateHandle OnDirectoryChangedHandle;
+	};
+	
+	/** Callback when a new content path is mounted */
+	void OnContentPathMounted(const FString& InAssetPath, const FString& InFilesystemPath);
+	/** Callback when an existing content path is unmounted */
+	void OnContentPathDismounted(const FString& InAssetPath, const FString& InFilesystemPath);
+
+	/** Register a mount path from a source content path */
+	void RegisterContentMountPath(const FString& InAssetPath, const FString& InFilesystemPath);
+	/** Unregister a mount path from a source content path */
+	void UnregisterContentMountPath(const FString& InAssetPath, const FString& InFilesystemPath);
+
+	/** Resolve the given path to its sandbox path (if any) */
+	FSandboxedPlatformFilePath ToSandboxPath(FString InFilename, const bool bEvenIfDisabled = false) const;
+	/** Resolve the given path to its sandbox path (if any) from an absolute filename */
+	FSandboxedPlatformFilePath ToSandboxPath_Absolute(FString InFilename, const bool bEvenIfDisabled = false) const;
+	/** Resolve the given path to its non-sandbox path (if any) */
+	FSandboxedPlatformFilePath FromSandboxPath(FString InFilename) const;
+	/** Resolve the given path to its non-sandbox path (if any) from an absolute filename */
+	FSandboxedPlatformFilePath FromSandboxPath_Absolute(FString InFilename) const;
+
+	/** Check whether the given absolute sandbox path has been explicitly deleted from the sandbox */
+	bool IsPathDeleted(const FSandboxedPlatformFilePath& InPath) const;
+	/** Set whether the given absolute sandbox path has been explicitly deleted from the sandbox */
+	void SetPathDeleted(const FSandboxedPlatformFilePath& InPath, const bool bIsDeleted);
+
+	/** Notify that a file has been explicitly deleted from the sandbox */
+	void NotifyFileDeleted(const FSandboxedPlatformFilePath& InPath);
+
+	/** Helper function to ensure that a sandbox contains a copy of the non-sandbox file (eg, prior to opening an existing file for writing) - does nothing if the sandbox already has the file, or if there is no non-sandbox file to copy */
+	void MigrateFileToSandbox(const FSandboxedPlatformFilePath& InPath) const;
+
+	/** Helper function to get the contents of a directory, taking into account the sandbox state - paths are returned relative to InDirBase */
+	TArray<FDirectoryItem> GetDirectoryContents(const FSandboxedPlatformFilePath& InPath, const TCHAR* InDirBase) const;
+
+#if WITH_EDITOR
+	/** Called when a file in a sandbox directory changes on disk */
+	void OnDirectoryChanged(const TArray<FFileChangeData>& FileChanges, FSandboxedPlatformFilePath MountPath);
+#endif
+
+	/** Copy file using source control (if enabled) */
+	bool CopyFileWithSCC(const TCHAR* InToFilename, const TCHAR* InFromFilename, const FPersistArgs& InParam, FPersistResult& OutResult);
+	/** Delete file using source control (if enabled) */
+	bool DeleteFileWithSCC(const TCHAR* InFilename, FPersistResult& OutResult);
+
+	FString GetSandboxRootPath() const;
+		
+	/** Root path of this sandbox */
+	const FString RootPath;
+
+	/** Underlying platform file that we're wrapping */
+	IPlatformFile* LowerLevel;
+
+	/** Is this sandbox currently enabled? */
+	TAtomic<bool> bSandboxEnabled;
+
+	/** Array of sandbox mount points */
+	TArray<FSandboxMountPoint> SandboxMountPoints;
+	/** Critical section protecting concurrent access to SandboxMountPoints */
+	mutable FCriticalSection SandboxMountPointsCS;
+
+	/** Set of absolute sandbox paths that have been explicitly deleted from the sandbox and shouldn't fallback to the non-sandbox items */
+	TSet<FSandboxedPlatformFilePath> DeletedSandboxPaths;
+	/** Critical section protecting concurrent access to DeletedSandboxPaths */
+	mutable FCriticalSection DeletedSandboxPathsCS;
+	
+	/** Invoked when a file is deleted. */
+	FDeletionStateChanged OnFileDeletedDelegate;
+};
+}
+
diff --git a/Engine/Plugins/Developer/Sandbox/FileSandbox/Source/FileSandboxCore/Private/Sandbox/Platform/SandboxedPlatformFilePath.h b/Engine/Plugins/Developer/Sandbox/FileSandbox/Source/FileSandboxCore/Private/Sandbox/Platform/SandboxedPlatformFilePath.h
new file mode 100644
--- /dev/null
+++ b/Engine/Plugins/Developer/Sandbox/FileSandbox/Source/FileSandboxCore/Private/Sandbox/Platform/SandboxedPlatformFilePath.h
@@ -0,0 +1,104 @@
+﻿// Copyright Epic Games, Inc. All Rights Reserved.
+
+#pragma once
+
+#include "Containers/UnrealString.h"
+
+namespace UE::FileSandboxCore
+{
+/** A pair of file paths: a path to the original file and a path to the sandboxed version, if it exists. */
+class FSandboxedPlatformFilePath
+{
+public:
+	
+	explicit FSandboxedPlatformFilePath(FString&& InNonSandboxPath)
+		: SandboxPath()
+		, NonSandboxPath(MoveTemp(InNonSandboxPath))
+	{
+	}
+
+	explicit FSandboxedPlatformFilePath(FString&& InNonSandboxPath, FString&& InSandboxPath)
+		: SandboxPath(MoveTemp(InSandboxPath))
+		, NonSandboxPath(MoveTemp(InNonSandboxPath))
+	{
+	}
+
+	static FSandboxedPlatformFilePath CreateSandboxPath(FString&& InNonSandboxPath, const FSandboxedPlatformFilePath& InRootPath)
+	{
+		checkf(InRootPath.HasSandboxPath(), TEXT("Root '%s' had no sandbox path set!"), *InRootPath.GetNonSandboxPath());
+		return CreateSandboxPath(MoveTemp(InNonSandboxPath), InRootPath.GetSandboxPath(), InRootPath.GetNonSandboxPath());
+	}
+
+	static FSandboxedPlatformFilePath CreateSandboxPath(FString&& InNonSandboxPath, const FString& InRootSandboxPath, const FString& InRootNonSandboxPath)
+	{
+		// Mount point are stored with a trailing slash to prevent matching mount point with similar names -> (/Bla/Content, /Bla/ContentSupreme)
+		// An extra slash is appended here to make sure with can match mount point directly -> (/Bla/Content match /Bla/Content/)
+		FString ResolvedSandboxPath = InNonSandboxPath + TEXT("/");
+		checkf(ResolvedSandboxPath.StartsWith(InRootNonSandboxPath), TEXT("Path '%s' was not under the root '%s'!"), *InNonSandboxPath, *InRootNonSandboxPath);
+		ResolvedSandboxPath.ReplaceInline(*InRootNonSandboxPath, *InRootSandboxPath);
+		ResolvedSandboxPath.RemoveAt(ResolvedSandboxPath.Len() - 1, EAllowShrinking::No);
+		return FSandboxedPlatformFilePath(MoveTemp(InNonSandboxPath), MoveTemp(ResolvedSandboxPath));
+	}
+
+	static FSandboxedPlatformFilePath CreateNonSandboxPath(FString&& InSandboxPath, const FSandboxedPlatformFilePath& InRootPath)
+	{
+		checkf(InRootPath.HasSandboxPath(), TEXT("Root '%s' had no sandbox path set!"), *InRootPath.GetNonSandboxPath());
+		return CreateNonSandboxPath(MoveTemp(InSandboxPath), InRootPath.GetSandboxPath(), InRootPath.GetNonSandboxPath());
+	}
+
+	static FSandboxedPlatformFilePath CreateNonSandboxPath(FString&& InSandboxPath, const FString& InRootSandboxPath, const FString& InRootNonSandboxPath)
+	{
+		// Mount point are stored with a trailing slash to prevent matching mount point with similar names -> (/Bla/Content, /Bla/ContentSupreme)
+		// An extra slash is appended here to make sure with can match mount point directly -> (/Bla/Content match /Bla/Content/)
+		FString ResolvedNonSandboxPath = InSandboxPath + TEXT("/");
+		checkf(ResolvedNonSandboxPath.StartsWith(InRootSandboxPath), TEXT("Path '%s' was not under the root '%s'!"), *InRootSandboxPath, *InRootSandboxPath);
+		ResolvedNonSandboxPath.ReplaceInline(*InRootSandboxPath, *InRootNonSandboxPath);
+		ResolvedNonSandboxPath.RemoveAt(ResolvedNonSandboxPath.Len() - 1, EAllowShrinking::No);
+		return FSandboxedPlatformFilePath(MoveTemp(ResolvedNonSandboxPath), MoveTemp(InSandboxPath));
+	}
+
+	/** Copyable */
+	FSandboxedPlatformFilePath(const FSandboxedPlatformFilePath&) = default;
+	FSandboxedPlatformFilePath& operator=(const FSandboxedPlatformFilePath&) = default;
+
+	/** Movable */
+	FSandboxedPlatformFilePath(FSandboxedPlatformFilePath&&) = default;
+	FSandboxedPlatformFilePath& operator=(FSandboxedPlatformFilePath&&) = default;
+
+	bool operator==(const FSandboxedPlatformFilePath& Rhs) const
+	{
+		return GetSandboxPath() == Rhs.GetSandboxPath() && GetNonSandboxPath() == Rhs.GetNonSandboxPath();
+	}
+
+	/** Do we have a sandbox path set? */
+	bool HasSandboxPath() const
+	{
+		return SandboxPath.Len() > 0;
+	}
+
+	/** Get the absolute sandbox path */
+	const FString& GetSandboxPath() const
+	{
+		return SandboxPath;
+	}
+
+	/** Get the absolute non-sandbox path */
+	const FString& GetNonSandboxPath() const
+	{
+		return NonSandboxPath;
+	}
+
+	friend uint32 GetTypeHash(const FSandboxedPlatformFilePath& InPath)
+	{
+		return HashCombine(GetTypeHash(InPath.SandboxPath), GetTypeHash(InPath.NonSandboxPath));
+	}
+
+private:
+	
+	/** Absolute sandbox path */
+	FString SandboxPath;
+
+	/** Absolute non-sandbox path */
+	FString NonSandboxPath;
+};
+}
diff --git a/Engine/Plugins/Developer/Sandbox/FileSandbox/Source/FileSandboxCore/Private/Sandbox/SandboxInstance.cpp b/Engine/Plugins/Developer/Sandbox/FileSandbox/Source/FileSandboxCore/Private/Sandbox/SandboxInstance.cpp
new file mode 100644
--- /dev/null
+++ b/Engine/Plugins/Developer/Sandbox/FileSandbox/Source/FileSandboxCore/Private/Sandbox/SandboxInstance.cpp
@@ -0,0 +1,162 @@
+﻿// Copyright Epic Games, Inc. All Rights Reserved.
+
+#include "SandboxInstance.h"
+
+#include "HAL/FileManager.h"
+#include "Misc/PackageName.h"
+#include "Types/Sandbox/PersistArgs.h"
+#include "Types/Manager/SandboxInitArgs.h"
+#include "Types/Sandbox/DiscardResult.h"
+#include "Utils/PackageSandboxUtils.h"
+#include "Utils/FileSandboxUtils.h"
+#include "Utils/SandboxDirectoryUtils.h"
+#include "Utils/SandboxUtils.h"
+
+namespace UE::FileSandboxCore
+{
+namespace Private
+{
+static void SetManifestFileFromRevertedFileChanges(FFileSandboxCore_ManifestData& OutManifestData, TConstArrayView<FFileChangeData> InFileChanges)
+{
+	OutManifestData.Empty();
+
+	for (const FFileChangeData& Change : InFileChanges)
+	{
+		switch (Change.Action)
+		{
+		case FFileChangeData::FCA_Added: OutManifestData.DeletedFiles.Add(Change.Filename); break;
+		case FFileChangeData::FCA_Modified: OutManifestData.ModifiedFiles.Add(Change.Filename); break;
+		case FFileChangeData::FCA_Removed: OutManifestData.AddedFiles.Add(Change.Filename); break;
+		default: break;
+		}
+	}
+}
+}
+	
+TUniquePtr<FSandboxInstance> FSandboxInstance::CreateNewSandbox(
+	FString InRootDirectory, const FSandboxInitArgs& InInitArgs, const FFileSandboxCore_SandboxMetaData& InMetaData
+	)
+{
+	FFileSandboxCore_ManifestData ManifestData;
+	const bool bSuccess = FileSandboxCore::SaveManifest(ManifestData, *InRootDirectory) && SaveMetaData(InMetaData, InRootDirectory);
+	return bSuccess ? MakeUnique<FSandboxInstance>(MoveTemp(InRootDirectory), InInitArgs, ManifestData, InMetaData) : nullptr;
+}
+
+TUniquePtr<FSandboxInstance> FSandboxInstance::LoadSandbox(FString InRootDirectory, const FSandboxInitArgs& InInitArgs)
+{
+	TOptional<FFileSandboxCore_ManifestData> ManifestData = LoadManifest(InRootDirectory);
+	TOptional<FFileSandboxCore_SandboxMetaData> MetaData = LoadMetaData(InRootDirectory);
+	const bool bSuccess = ManifestData && MetaData;
+	return bSuccess ? MakeUnique<FSandboxInstance>(MoveTemp(InRootDirectory), InInitArgs, MoveTemp(*ManifestData), MoveTemp(*MetaData)) : nullptr;
+}
+
+FSandboxInstance::FSandboxInstance(
+	FString InRootDirectory, const FSandboxInitArgs& InitArgs,
+	FFileSandboxCore_ManifestData InManifestData, FFileSandboxCore_SandboxMetaData InMetaData
+	)
+	: RootDirectory(FPaths::ConvertRelativePathToFull(MoveTemp(InRootDirectory)))
+	, ManifestContent(MoveTemp(InManifestData))
+	, MetaData(MoveTemp(InMetaData))
+	, PlatformFile(MakeUnique<FSandboxPlatformFile>(RootDirectory))
+{
+	PlatformFile->Initialize(&InitArgs.GetPlatformFile(), TEXT(""));
+	ReplayChanges(ManifestContent.DeletedFiles.GetFiles(), ManifestContent.ModifiedFiles.GetFiles(), ManifestContent.AddedFiles.GetFiles());
+	PlatformFile->OnDeletionStateChanged().AddRaw(this, &FSandboxInstance::OnFileDeletionStateChanged);
+}
+	
+FSandboxInstance::~FSandboxInstance()
+{
+	TArray<FName> PackagesPendingHotReload, PackagesPendingPurge;
+	TArray<FFileChangeData> FileChanges;
+	PlatformFile->DiscardSandbox(PackagesPendingHotReload, PackagesPendingPurge, &FileChanges);
+	PlatformFile.Reset();
+
+	// Save the data before hot reloading or purging: in case that causes a crash we'll have an update to manifest file. 
+	Private::SetManifestFileFromRevertedFileChanges(ManifestContent, FileChanges);
+	SaveManifest();
+	// This just exists to handle e.g. the user having deleted the file manually during the sandbox. Normally, the metadata file should still be there.
+	if (!FPaths::FileExists(RootDirectory / GetMetadataFileName()))
+	{
+		SaveMetaData(MetaData, RootDirectory);
+	}
+
+	if (!IsEngineExitRequested())
+	{
+		// Hot reload after unregistering from most delegates to prevent events triggered by hot-reloading (such as asset deleted) to be recorded as transaction.
+		PurgePackages(PackagesPendingPurge);
+		HotReloadPackages(PackagesPendingHotReload);
+	}
+}
+
+FPersistResult FSandboxInstance::PersistSandbox(const FPersistArgs& InArgs)
+{
+	bool bManifestNeedsSave = false;
+	const FPersistResult Result = PlatformFile->PersistSandbox(InArgs,
+		[this, &bManifestNeedsSave](const FSandboxedPlatformFilePath& InPath, FSandboxPlatformFile::EPersistedFileType InPersistType)
+		{
+			switch (InPersistType)
+			{
+			case FSandboxPlatformFile::EPersistedFileType::Deleted:
+				bManifestNeedsSave |= ManifestContent.DeletedFiles.Remove(InPath.GetNonSandboxPath());
+				break;
+			case FSandboxPlatformFile::EPersistedFileType::Moved:
+				bManifestNeedsSave |= ManifestContent.ModifiedFiles.Remove(InPath.GetNonSandboxPath());
+				bManifestNeedsSave |= ManifestContent.AddedFiles.Remove(InPath.GetNonSandboxPath());
+				break;
+				
+			default: checkNoEntry(); break;
+			}
+		});
+
+	if (bManifestNeedsSave)
+	{
+		SaveManifest();
+	}
+
+	return Result;
+}
+
+FDiscardResult FSandboxInstance::DiscardSandbox(const FDiscardArgs& InArgs)
+{
+	FDiscardResult Result;
+	PlatformFile->DiscardSandbox(Result.PackagesPendingHotReload, Result.PackagesPendingPurge);
+	
+	ManifestContent.Empty();
+	SaveManifest();
+
+	return Result;
+}
+
+TArray<FString> FSandboxInstance::GatherChangedFiles() const
+{
+	return PlatformFile->GatherSandboxChangedFilenames([this](const FSandboxedPlatformFilePath& InFilePath)
+	{
+		// TODO DP: Do we need this?
+		/*const FString& NonSandboxPath = InFilePath.GetNonSandboxPath();
+		const FDateTime* ModifiedTime = ManifestContent.PersistedFiles.Find(NonSandboxPath);
+		const bool bHasBeenPersisted = ModifiedTime && PlatformFile->GetTimeStamp(*NonSandboxPath) <= *ModifiedTime;
+		return !bHasBeenPersisted;*/
+		return false;
+	});
+}
+
+bool FSandboxInstance::DeletedPackageExistsInNonSandbox(const FString& InFilename) const
+{
+	return PlatformFile->DeletedPackageExistsInNonSandbox(InFilename);
+}
+
+void FSandboxInstance::OnFileDeletionStateChanged(const FSandboxedPlatformFilePath& InFile, bool bIsDeleted)
+{
+	const bool bNeedsSave = (bIsDeleted && ManifestContent.DeletedFiles.Add(InFile.GetNonSandboxPath()))
+		|| (!bIsDeleted && ManifestContent.DeletedFiles.Remove(InFile.GetNonSandboxPath()));
+	if (bNeedsSave)
+	{
+		SaveManifest();
+	}
+}
+
+bool FSandboxInstance::SaveManifest() const
+{
+	return FileSandboxCore::SaveManifest(ManifestContent, GetRootDirectory());
+}
+}
diff --git a/Engine/Plugins/Developer/Sandbox/FileSandbox/Source/FileSandboxCore/Private/Sandbox/SandboxInstance.h b/Engine/Plugins/Developer/Sandbox/FileSandbox/Source/FileSandboxCore/Private/Sandbox/SandboxInstance.h
new file mode 100644
--- /dev/null
+++ b/Engine/Plugins/Developer/Sandbox/FileSandbox/Source/FileSandboxCore/Private/Sandbox/SandboxInstance.h
@@ -0,0 +1,74 @@
+﻿// Copyright Epic Games, Inc. All Rights Reserved.
+
+#pragma once
+
+#include "ISandboxInstance.h"
+#include "Data/ManifestData.h"
+#include "Data/SandboxMetaData.h"
+#include "Platform/SandboxPlatformFile.h"
+#include "Templates/UnrealTemplate.h"
+
+namespace UE::FileSandboxCore
+{
+}
+
+namespace UE::FileSandboxCore
+{
+struct FSandboxInitArgs;
+	
+class FSandboxInstance : public ISandboxInstance, public FNoncopyable
+{
+public:
+
+	/** Initializes InRootDirectory with a new sandbox. */
+	static TUniquePtr<FSandboxInstance> CreateNewSandbox(
+		FString InRootDirectory, const FSandboxInitArgs& InInitArgs, const FFileSandboxCore_SandboxMetaData& InMetaData
+		);
+	/** Loads a preexisting sandbox from InRootDirectory. */
+	static TUniquePtr<FSandboxInstance> LoadSandbox(FString InRootDirectory, const FSandboxInitArgs& InInitArgs);
+
+	explicit FSandboxInstance(
+		FString InRootDirectory,
+		const FSandboxInitArgs& InitArgs,
+		FFileSandboxCore_ManifestData InManifestData,
+		FFileSandboxCore_SandboxMetaData InMetaData
+		);
+	virtual ~FSandboxInstance() override;
+
+	using ISandboxInstance::DiscardSandbox;
+
+	//~ Begin ISandboxInstance Interface
+	virtual FPersistResult PersistSandbox(const FPersistArgs& InArgs) override;
+	virtual FDiscardResult DiscardSandbox(const FDiscardArgs& InArgs) override;
+	virtual TArray<FString> GatherChangedFiles() const override;
+	virtual bool DeletedPackageExistsInNonSandbox(const FString& InFilename) const override;
+	virtual const FFileSandboxCore_SandboxMetaData& GetInitialMetaData() const override { return MetaData; }
+	virtual const TCHAR* GetRootDirectory() const override { return *RootDirectory; }
+	//~ End ISandboxInstance Interface
+
+private:
+
+	/** Root directory in which we save all out content. */
+	const FString RootDirectory;
+
+	/** Current manifest data. */
+	FFileSandboxCore_ManifestData ManifestContent;
+	/**
+	 * Metadata about this sandbox.
+	 * 
+	 * The sandbox does not actively use the metadata, but we keep it in memory in case anything happens to the file, e.g. the user manually deleting
+	 * it while in the sandbox. This way we can always guarantee to the API user that there is always valid metadata.
+	 */
+	const FFileSandboxCore_SandboxMetaData MetaData;
+
+	/** Hooks into the engine's I/O system to redirect saving & loading of files. */
+	TUniquePtr<FSandboxPlatformFile> PlatformFile;
+
+	/** Marks the file as deleted in the manifest. */
+	void OnFileDeletionStateChanged(const FSandboxedPlatformFilePath& InFile, bool bIsDeleted);
+
+	/** Saves the manifest file. */
+	bool SaveManifest() const;
+};
+}
+
diff --git a/Engine/Plugins/Developer/Sandbox/FileSandbox/Source/FileSandboxCore/Private/Sandbox/SandboxManager.cpp b/Engine/Plugins/Developer/Sandbox/FileSandbox/Source/FileSandboxCore/Private/Sandbox/SandboxManager.cpp
new file mode 100644
--- /dev/null
+++ b/Engine/Plugins/Developer/Sandbox/FileSandbox/Source/FileSandboxCore/Private/Sandbox/SandboxManager.cpp
@@ -0,0 +1,113 @@
+﻿// Copyright Epic Games, Inc. All Rights Reserved.
+
+#include "SandboxManager.h"
+
+#include "LogFileSandbox.h"
+#include "GenericPlatform/GenericPlatformFile.h"
+#include "HAL/PlatformFileManager.h"
+#include "Misc/CoreDelegates.h"
+#include "Misc/Paths.h"
+#include "Types/Manager/DeleteSandboxResult.h"
+#include "Types/Manager/LoadSandboxByDirectoryArgs.h"
+#include "Types/Manager/NewSandboxArgs.h"
+#include "Utils/FileSandboxUtils.h"
+
+namespace UE::FileSandboxCore
+{
+FSandboxManager::FSandboxManager()
+{
+	FCoreDelegates::OnEnginePreExit.AddRaw(this, &FSandboxManager::OnEngineExit);
+}
+
+FSandboxManager::~FSandboxManager()
+{
+	FCoreDelegates::OnEnginePreExit.RemoveAll(this);
+}
+
+FNewSandboxResult FSandboxManager::CreateNewSandbox(const FNewSandboxArgs& InArgs)
+{
+	const FString AbsoluteRoot = FPaths::ConvertRelativePathToFull(InArgs.SandboxBasePath);
+	
+	const FString& Name = InArgs.MetaData.Name;
+	if (DetermineDirectorySuitability(Name, InArgs.SandboxBasePath) != EDirectorySuitability::Suitable)
+	{
+		return { MakeError(FNewSandboxError{ ENewSandboxErrorReason::UnsuitablePath }) };
+	}
+
+	LeaveSandbox();
+	
+	const FString NewDirectory = FPaths::Combine(InArgs.SandboxBasePath, Name);
+	ActiveSandbox = FSandboxInstance::CreateNewSandbox(NewDirectory, InArgs.InitArgs, InArgs.MetaData);
+	ISandboxInstance* Instance = ActiveSandbox.Get();
+	if (!Instance)
+	{
+		return { MakeError(FNewSandboxError{ ENewSandboxErrorReason::IOError }) };
+	}
+	
+	OnPostSandboxStartupDelegate.Broadcast(*Instance);
+	return { MakeValue(Instance) };
+}
+
+FLoadSandboxResult FSandboxManager::LoadSandbox(const FLoadSandboxByDirectoryArgs& InLoadArgs)
+{
+	const FString AbsoluteRoot = FPaths::ConvertRelativePathToFull(InLoadArgs.SandboxRootPath);
+	if (!IsRootSandboxDirectory(AbsoluteRoot))
+	{
+		UE_LOG(LogFileSandbox, Error, TEXT("Directory %s does not contain any sandbox"), *AbsoluteRoot);
+		return { MakeError(FLoadSandboxError{ ELoadSandboxLoadErrorReason::InvalidDirectory }) };
+	}
+
+	// This could happen if the sandbox is placed in the root e.g. "D:/".
+	// It means that the sandbox was not created in GetBaseSandboxDirectory. 
+	// Names are just a simplified concept for managing sandboxes that are placed in this default folder.
+	// This is a custom sandbox so it's beyond the default workflow.
+	const FString LastDirectoryName = FPaths::GetPathLeaf(AbsoluteRoot);
+	UE_CLOG(LastDirectoryName.IsEmpty(), LogFileSandbox, Warning, TEXT("Sandbox placed %s will not have any name."), *AbsoluteRoot);
+	const TOptional<FString> SandboxName = LastDirectoryName.IsEmpty() ? LastDirectoryName : TOptional<FString>{};
+
+	LeaveSandbox();
+	
+	ActiveSandbox = FSandboxInstance::LoadSandbox(AbsoluteRoot, InLoadArgs.InitArgs);
+	ISandboxInstance* Instance = ActiveSandbox.Get();
+	if (!Instance)
+	{
+		return { MakeError(FLoadSandboxError{ ELoadSandboxLoadErrorReason::IOError }) };
+	}
+	
+	OnPostSandboxStartupDelegate.Broadcast(*Instance);
+	return { MakeValue(Instance) };
+}
+
+FDeleteSandboxResult FSandboxManager::DeleteSandbox(const FString& InSandboxDirectory)
+{
+	if (!IsRootSandboxDirectory(InSandboxDirectory))
+	{
+		return EDeleteSandboxErrorCode::InvalidDirectory;
+	}
+	
+	if (ActiveSandbox && FPaths::IsSamePath(ActiveSandbox->GetRootDirectory(), InSandboxDirectory))
+	{
+		LeaveSandbox();
+	}
+
+	IPlatformFile& PlatformFile = FPlatformFileManager::Get().GetPlatformFile();
+	return PlatformFile.DeleteDirectoryRecursively(*InSandboxDirectory)
+		? EDeleteSandboxErrorCode::Success
+		: EDeleteSandboxErrorCode::IOError;
+}
+
+void FSandboxManager::LeaveSandbox(const FLeaveSandboxArgs& InLeaveArgs)
+{
+	if (ISandboxInstance* Instance = ActiveSandbox.Get())
+	{
+		OnPreSandboxShutdownDelegate.Broadcast(*Instance);
+		ActiveSandbox.Reset();
+		OnPostSandboxShutdownDelegate.Broadcast();
+	}
+}
+
+void FSandboxManager::OnEngineExit()
+{
+	LeaveSandbox();
+}
+}
diff --git a/Engine/Plugins/Developer/Sandbox/FileSandbox/Source/FileSandboxCore/Private/Sandbox/SandboxManager.h b/Engine/Plugins/Developer/Sandbox/FileSandbox/Source/FileSandboxCore/Private/Sandbox/SandboxManager.h
new file mode 100644
--- /dev/null
+++ b/Engine/Plugins/Developer/Sandbox/FileSandbox/Source/FileSandboxCore/Private/Sandbox/SandboxManager.h
@@ -0,0 +1,47 @@
+﻿// Copyright Epic Games, Inc. All Rights Reserved.
+
+#pragma once
+
+#include "ISandboxManager.h"
+#include "SandboxInstance.h"
+#include "Templates/UnrealTemplate.h"
+
+namespace UE::FileSandboxCore
+{
+/** Manges the sandbox lifetime for the module */
+class FSandboxManager : public ISandboxManager, public FNoncopyable
+{
+public:
+
+	FSandboxManager();
+	virtual ~FSandboxManager() override;
+
+	//~ Begin ISandboxManager Interface
+	virtual FNewSandboxResult CreateNewSandbox(const FNewSandboxArgs& InArgs) override;
+	virtual FLoadSandboxResult LoadSandbox(const FLoadSandboxByDirectoryArgs& InLoadArgs) override;
+	virtual FDeleteSandboxResult DeleteSandbox(const FString& InSandboxDirectory) override;
+	virtual void LeaveSandbox(const FLeaveSandboxArgs& InLeaveArgs = FLeaveSandboxArgs()) override;
+	virtual ISandboxInstance* GetActiveSandboxInstance() const override { return ActiveSandbox.Get(); }
+	virtual FSandboxInstanceEvent& OnPostSandboxStartup() override { return OnPostSandboxStartupDelegate; }
+	virtual FSandboxInstanceEvent& OnPreSandboxShutdown() override { return OnPreSandboxShutdownDelegate; }
+	virtual FSandboxShutdownEvent& OnPostSandboxShutdown() override { return OnPostSandboxShutdownDelegate; }
+	//~ End ISandboxManager Interface
+
+private:
+
+	/** Current sandbox. Nullptr if not in any sandbox. */
+	TUniquePtr<FSandboxInstance> ActiveSandbox;
+
+	/** Invoked after a sandbox has been successfully created. */
+	FSandboxInstanceEvent OnPostSandboxStartupDelegate;
+	/** Invoked when it has been decided, that a sandbox instance will be taken down, but before any core destruction logic runs. */
+	FSandboxInstanceEvent OnPreSandboxShutdownDelegate;
+	/** Invoked at the end of shutting down a sandbox, i.e. after all core destruction logic has run. */
+	FSandboxShutdownEvent OnPostSandboxShutdownDelegate;
+
+	/** Stops the current sandbox. */
+	void OnEngineExit();
+};
+}
+
+
diff --git a/Engine/Plugins/Developer/Sandbox/FileSandbox/Source/FileSandboxCore/Private/Utils/FileSandboxUtils.cpp b/Engine/Plugins/Developer/Sandbox/FileSandbox/Source/FileSandboxCore/Private/Utils/FileSandboxUtils.cpp
new file mode 100644
--- /dev/null
+++ b/Engine/Plugins/Developer/Sandbox/FileSandbox/Source/FileSandboxCore/Private/Utils/FileSandboxUtils.cpp
@@ -0,0 +1,56 @@
+﻿// Copyright Epic Games, Inc. All Rights Reserved.
+
+#include "FileSandboxUtils.h"
+
+#include "JsonObjectConverter.h"
+#include "LogFileSandbox.h"
+#include "Data/ManifestData.h"
+#include "Misc/FileHelper.h"
+#include "Misc/Optional.h"
+#include "Misc/Paths.h"
+
+namespace UE::FileSandboxCore
+{
+bool SaveManifest(const FFileSandboxCore_ManifestData& InData, const FString& InRootDirectory)
+{
+	FString JSONPayload;
+	FJsonObjectConverter::UStructToJsonObjectString(InData, JSONPayload, 0, 0);
+	const FString AbsFileName = InRootDirectory / GetManifestFileName();
+	const bool bSuccess = FFileHelper::SaveStringToFile(JSONPayload, *AbsFileName);
+	UE_CLOG(!bSuccess, LogFileSandbox, Error, TEXT("Failed to write sandbox manifest file %s"), *InRootDirectory);
+	return bSuccess;
+}
+
+TOptional<FFileSandboxCore_ManifestData> LoadManifest(const FString& InRootDirectory, bool bLogErrorIfNotExist)
+{
+	FFileSandboxCore_ManifestData ManifestPlatformFileState;
+	FString JSONPayload;
+	
+	const FString AbsFileName = InRootDirectory / GetManifestFileName();
+	if (FFileHelper::LoadFileToString(JSONPayload, *AbsFileName))
+	{
+		const bool bSuccess = FJsonObjectConverter::JsonObjectStringToUStruct(JSONPayload, &ManifestPlatformFileState, 0, 0);
+		UE_CLOG(!bSuccess, LogFileSandbox, Error, TEXT("Failed to parse sandbox manifest file %s"), *InRootDirectory);
+		return bSuccess ? ManifestPlatformFileState : TOptional<FFileSandboxCore_ManifestData>{};
+	}
+	
+	UE_CLOG(bLogErrorIfNotExist, LogFileSandbox, Error, TEXT("Failed to load sandbox manifest file %s"), *InRootDirectory);
+	return {};
+}
+
+bool IsRootSandboxDirectory(const FString& InRootDirectory)
+{
+	return FPaths::FileExists(InRootDirectory / GetManifestFileName())
+		&& FPaths::FileExists(InRootDirectory / GetMetadataFileName());
+}
+
+FString GetManifestFileName()
+{
+	return TEXT("manifest.json");
+}
+
+FString GetMetadataFileName()
+{
+	return TEXT("metadata.json");
+}
+}
diff --git a/Engine/Plugins/Developer/Sandbox/FileSandbox/Source/FileSandboxCore/Private/Utils/FileSandboxUtils.h b/Engine/Plugins/Developer/Sandbox/FileSandbox/Source/FileSandboxCore/Private/Utils/FileSandboxUtils.h
new file mode 100644
--- /dev/null
+++ b/Engine/Plugins/Developer/Sandbox/FileSandbox/Source/FileSandboxCore/Private/Utils/FileSandboxUtils.h
@@ -0,0 +1,23 @@
+﻿// Copyright Epic Games, Inc. All Rights Reserved.
+
+#pragma once
+
+class FString;
+struct FFileSandboxCore_ManifestData;
+template<typename OptionalType> struct TOptional;
+
+namespace UE::FileSandboxCore
+{
+/** Saves the manifest file. */
+bool SaveManifest(const FFileSandboxCore_ManifestData& InData, const FString& InRootDirectory);
+/** Loads the manifest file. */
+TOptional<FFileSandboxCore_ManifestData> LoadManifest(const FString& InRootDirectory, bool bLogErrorIfNotExist = true);
+
+/** @return Whether this is a path to the root of a sandbox, i.e. whether it contains the manifest and the metadata file. */
+bool IsRootSandboxDirectory(const FString& InRootDirectory);
+
+/* @return The name to give the manifest file. */
+FString GetManifestFileName();
+/* @return The name to give the metadata file. */
+FString GetMetadataFileName();
+}
\ No newline at end of file
diff --git a/Engine/Plugins/Developer/Sandbox/FileSandbox/Source/FileSandboxCore/Private/Utils/PackageSandboxUtils.cpp b/Engine/Plugins/Developer/Sandbox/FileSandbox/Source/FileSandboxCore/Private/Utils/PackageSandboxUtils.cpp
new file mode 100644
--- /dev/null
+++ b/Engine/Plugins/Developer/Sandbox/FileSandbox/Source/FileSandboxCore/Private/Utils/PackageSandboxUtils.cpp
@@ -0,0 +1,219 @@
+﻿// Copyright Epic Games, Inc. All Rights Reserved.
+
+#include "Utils/PackageSandboxUtils.h"
+
+#include "Engine/GameEngine.h"
+#include "Engine/Level.h"
+#include "Engine/World.h"
+#include "GameMapsSettings.h"
+#include "Misc/MessageDialog.h"
+#include "RenderingThread.h"
+#include "UObject/LinkerLoad.h"
+
+#if WITH_EDITOR
+#include "Editor.h"
+#include "FileHelpers.h"
+#include "Selection.h"
+#include "Subsystems/AssetEditorSubsystem.h"
+#endif
+
+namespace UE::FileSandboxCore
+{
+	
+UWorld* GetCurrentWorld()
+{
+	UWorld* CurrentWorld = nullptr;
+#if WITH_EDITOR
+	if (GIsEditor)
+	{
+		CurrentWorld = GEditor->GetEditorWorldContext().World();
+	}
+	else
+#endif
+		if (UGameEngine* GameEngine = Cast<UGameEngine>(GEngine))
+		{
+			CurrentWorld = GameEngine->GetGameWorld();
+		}
+	return CurrentWorld;
+}
+
+bool ShouldReloadPersistentLevel(UPackage *PackageToReload)
+{
+	UWorld* CurrentWorld = GetCurrentWorld();
+	if (CurrentWorld)
+	{
+		const TArray<ULevel*> Levels = CurrentWorld->GetLevels();
+		for (ULevel* Level : Levels)
+		{
+			if (Level->GetPackage() == PackageToReload)
+			{
+				return true;
+			}
+		}
+	}
+
+	return false;
+}
+	
+void PurgePackages(const TArray<FName>& InPackageNames)
+{
+	if (InPackageNames.Num() == 0)
+	{
+		return;
+	}
+
+#if WITH_EDITOR
+	TArray<UObject*> ObjectsToPurge;
+	auto CollectObjectToPurge = [&ObjectsToPurge](UObject* InObject)
+	{
+		if (InObject->IsAsset() && GIsEditor)
+		{
+			GEditor->GetEditorSubsystem<UAssetEditorSubsystem>()->CloseAllEditorsForAsset(InObject);
+			GEditor->GetSelectedObjects()->Deselect(InObject);
+		}
+		ObjectsToPurge.Add(InObject);
+	};
+
+	// Get the current edited map package to check if its going to be purged.
+	bool bEditedMapPurged = false;
+	UWorld* CurrentWorld = GetCurrentWorld();
+	UPackage* EditedMapPackage = CurrentWorld ? CurrentWorld->GetOutermost(): nullptr;
+
+	// Collect any in-memory packages that should be purged and check if we are including the current map in the purge.
+	for (const FName& PackageName : InPackageNames)
+	{
+		UPackage* ExistingPackage = FindPackage(nullptr, *PackageName.ToString());
+		if (ExistingPackage)
+		{
+			// Prevent any message from the editor saying a package is not saved or doesn't exist on disk.
+			ExistingPackage->SetDirtyFlag(false);
+
+			CollectObjectToPurge(ExistingPackage);
+			ForEachObjectWithPackage(ExistingPackage, [&CollectObjectToPurge](UObject* InObject)
+			{
+				CollectObjectToPurge(InObject);
+				return true;
+			});
+
+			bEditedMapPurged |= EditedMapPackage == ExistingPackage;
+		}
+	}
+
+	// Broadcast the eminent objects destruction (ex. tell BlueprintActionDatabase to release its reference(s) on Blueprint(s) right now)
+	FEditorDelegates::OnAssetsPreDelete.Broadcast(ObjectsToPurge);
+
+	// Mark objects as purgeable.
+	for (UObject* Object : ObjectsToPurge)
+	{
+		if (Object->IsRooted())
+		{
+			Object->RemoveFromRoot();
+		}
+		Object->ClearFlags(RF_Public | RF_Standalone);
+	}
+
+	// TODO: Revisit force replacing reference, current implementation is too aggressive and causes instability
+	// If we have any object that were made purgeable, null out their references so we can garbage collect
+	//if (ObjectsToPurge.Num() > 0)
+	//{
+	//	ObjectTools::ForceReplaceReferences(nullptr, ObjectsToPurge);
+	//
+	//}
+
+	// Check if the map being edited is going to be purged. (b/c it's being deleted)
+	if (bEditedMapPurged)
+	{
+		// The world being edited was purged and cannot be saved anymore, even with 'Save Current As', replace it by something sensible.
+		FString StartupMapPackage = GetDefault<UGameMapsSettings>()->EditorStartupMap.GetLongPackageName();
+		if (FPackageName::DoesPackageExist(StartupMapPackage))
+		{
+			UEditorLoadingAndSavingUtils::NewMapFromTemplate(StartupMapPackage, /*bSaveExistingMap*/false); // Expected to run GC internally.
+		}
+		else
+		{
+			UEditorLoadingAndSavingUtils::NewBlankMap(/*bSaveExistingMap*/false); // Expected to run GC internally.
+		}
+	}
+	// if we have object to purge but the map isn't one of them collect garbage (if we purged the map it has already been done)
+	else if (ObjectsToPurge.Num() > 0)
+	{
+		CollectGarbage(GARBAGE_COLLECTION_KEEPFLAGS);
+	}
+#endif // WITH_EDITOR
+
+}
+
+void HotReloadPackages(const TArray<FName>& InPackageNames)
+{
+	if (InPackageNames.Num() == 0)
+	{
+		return;
+	}
+
+#if WITH_EDITOR
+	// Flush loading and clean-up any temporary placeholder packages (due to a package previously being missing on disk)
+	FlushAsyncLoading();
+	{
+		bool bRunGC = false;
+		for (const FName& PackageName : InPackageNames)
+		{
+			bRunGC |= FLinkerLoad::RemoveKnownMissingPackage(PackageName);
+		}
+		if (bRunGC)
+		{
+			CollectGarbage(GARBAGE_COLLECTION_KEEPFLAGS);
+		}
+	}
+
+	bool bAddPersistentLevel = false;
+
+	// Find the packages in-memory to content hot-reload
+	TArray<UPackage*> ExistingPackages;
+	ExistingPackages.Reserve(InPackageNames.Num());
+
+	for (const FName& PackageName : InPackageNames)
+	{
+		UPackage* ExistingPackage = FindPackage(nullptr, *PackageName.ToString());
+		if (ExistingPackage)
+		{
+			if (ExistingPackage->HasAnyPackageFlags(PKG_NewlyCreated))
+			{
+				ExistingPackage->ClearPackageFlags(PKG_NewlyCreated);
+			}
+			if (ExistingPackage->ContainsMap())
+			{
+				bAddPersistentLevel = ShouldReloadPersistentLevel(ExistingPackage);
+			}
+
+			if (!ExistingPackage->ContainsMap() || ExistingPackage->IsDirty())
+			{
+				ExistingPackages.Add(ExistingPackage);
+			}
+		}
+	}
+
+	UWorld* CurrentWorld = GetCurrentWorld();
+	if (CurrentWorld && bAddPersistentLevel)
+	{
+		ULevel* PersistentLevel = CurrentWorld->PersistentLevel;
+		if (PersistentLevel && !ExistingPackages.Contains(PersistentLevel->GetPackage()))
+		{
+			ExistingPackages.Add(PersistentLevel->GetPackage());
+		}
+	}
+
+	if (ExistingPackages.Num() > 0)
+	{
+		FlushRenderingCommands();
+
+		FText ErrorMessage;
+		UPackageTools::ReloadPackages(ExistingPackages, ErrorMessage, UPackageTools::EReloadPackagesInteractionMode::AssumePositive);
+
+		if (!ErrorMessage.IsEmpty())
+		{
+			FMessageDialog::Open(EAppMsgType::Ok, ErrorMessage);
+		}
+	}
+#endif
+}
+}
diff --git a/Engine/Plugins/Developer/Sandbox/FileSandbox/Source/FileSandboxCore/Private/Utils/SandboxDirectoryUtils.cpp b/Engine/Plugins/Developer/Sandbox/FileSandbox/Source/FileSandboxCore/Private/Utils/SandboxDirectoryUtils.cpp
new file mode 100644
--- /dev/null
+++ b/Engine/Plugins/Developer/Sandbox/FileSandbox/Source/FileSandboxCore/Private/Utils/SandboxDirectoryUtils.cpp
@@ -0,0 +1,198 @@
+﻿// Copyright Epic Games, Inc. All Rights Reserved.
+
+#include "Utils/SandboxDirectoryUtils.h"
+
+#include "FileSandboxUtils.h"
+#include "JsonObjectConverter.h"
+#include "LogFileSandbox.h"
+#include "Data/SandboxMetaData.h"
+#include "HAL/PlatformFileManager.h"
+#include "Misc/FileHelper.h"
+#include "Misc/Paths.h"
+#include "Templates/Function.h"
+#include "Types/EBreakBehavior.h"
+
+#define LOCTEXT_NAMESPACE "SandboxDirectoryUtils"
+
+namespace UE::FileSandboxCore
+{
+FString GetBaseSandboxDirectory()
+{
+	return FPaths::ProjectIntermediateDir() / TEXT("Sandboxes");
+}
+	
+void ForEachSandbox(
+	TFunctionRef<FProcessSandboxDirectorySignature> InProcessSandbox,
+	const FString& InBaseDirectory
+	)
+{
+	IPlatformFile& PlatformFile = FPlatformFileManager::Get().GetPlatformFile();
+	PlatformFile.IterateDirectory(*InBaseDirectory, [&InProcessSandbox](const TCHAR* FilenameOrDirectory, const bool bIsDirectory)
+	{
+		if (!bIsDirectory || !IsRootSandboxDirectory(FilenameOrDirectory))
+		{
+			return true;
+		}
+			
+		return InProcessSandbox(FilenameOrDirectory) == EBreakBehavior::Continue;
+	});
+}
+
+namespace Private
+{
+static FFileSandboxCore_SandboxMetaData MaskMetaData(const FFileSandboxCore_SandboxMetaData& InMetaData, EMetaDataWriteFlags InFlags)
+{
+	FFileSandboxCore_SandboxMetaData Result;
+
+	if (EnumHasAnyFlags(InFlags, EMetaDataWriteFlags::Name))
+	{
+		Result.Name = InMetaData.Name;
+	}
+	if (EnumHasAnyFlags(InFlags, EMetaDataWriteFlags::Description))
+	{
+		Result.Description = InMetaData.Description;
+	}
+	if (EnumHasAnyFlags(InFlags, EMetaDataWriteFlags::Tags))
+	{
+		Result.Tags = InMetaData.Tags;
+	}
+	if (EnumHasAnyFlags(InFlags, EMetaDataWriteFlags::CustomData))
+	{
+		Result.CustomData = InMetaData.CustomData;
+	}
+
+	return Result;
+}
+}
+
+bool SaveMetaData(
+	const FFileSandboxCore_SandboxMetaData& NewMetaData, const FString& InSandboxDirectory, EMetaDataWriteFlags InWriteFlags
+	)
+{
+	FString JSONPayload;
+	FJsonObjectConverter::UStructToJsonObjectString(
+		EnumHasAllFlags(InWriteFlags, EMetaDataWriteFlags::All) ? NewMetaData : Private::MaskMetaData(NewMetaData, InWriteFlags),
+		JSONPayload, 0, 0
+		);
+	
+	const FString AbsFileName = InSandboxDirectory / GetMetadataFileName();
+	const bool bSuccess =  FFileHelper::SaveStringToFile(JSONPayload, *AbsFileName);
+	UE_CLOG(!bSuccess, LogFileSandbox, Error, TEXT("Failed to write sandbox metadata file %s"), *AbsFileName);
+	return bSuccess;
+}
+	
+TOptional<FFileSandboxCore_SandboxMetaData> LoadMetaData(const FString& InSandboxDirectory)
+{
+	FFileSandboxCore_SandboxMetaData ManifestPlatformFileState;
+	FString JSONPayload;
+	
+	const FString AbsFileName = InSandboxDirectory / GetMetadataFileName();
+	if (FFileHelper::LoadFileToString(JSONPayload, *AbsFileName))
+	{
+		const bool bParsed = FJsonObjectConverter::JsonObjectStringToUStruct(JSONPayload, &ManifestPlatformFileState, 0, 0);
+		UE_CLOG(!bParsed, LogFileSandbox, Error, TEXT("Failed to parse sandbox metadata file %s"), *AbsFileName);
+		return bParsed ? ManifestPlatformFileState : TOptional<FFileSandboxCore_SandboxMetaData>{};
+	}
+	return {};
+}
+
+TOptional<FString> FindSandboxByName(
+	const FString& InSandboxName,
+	const FString& InBaseDirectory
+	)
+{
+	TOptional<FString> Result;
+
+	// First we'll look for a directory with the sandbox name and only load that metadata file... because that is the common case...
+	ForEachSandbox([&InSandboxName, &Result](const FString& InPath)
+	{
+		if (FPaths::GetPathLeaf(InPath) != InSandboxName)
+		{
+			return EBreakBehavior::Continue;
+		}
+
+		const TOptional<FFileSandboxCore_SandboxMetaData> MetaData = LoadMetaData(InPath);
+		if (MetaData && MetaData->Name == InSandboxName)
+		{
+			Result = InPath;
+		}
+			
+		return EBreakBehavior::Break;
+	}, InBaseDirectory);
+
+	if (Result)
+	{
+		return Result;
+	}
+
+	// ... but somebody may have messed with the file system (user could have renamed folder or edited metadata file directly).
+	ForEachSandbox([&InSandboxName, &Result](const FString& InPath)
+	{
+		const TOptional<FFileSandboxCore_SandboxMetaData> MetaData = LoadMetaData(InPath);
+		if (MetaData && MetaData->Name == InSandboxName)
+		{
+			Result = InPath;
+			return EBreakBehavior::Break;
+		}
+		return EBreakBehavior::Continue;
+	}, InBaseDirectory);
+	
+	return Result;
+}
+
+EDirectorySuitability DetermineDirectorySuitability(const FString& InSandboxName, const FString& InBaseDirectory, FText* OutReason)
+{
+	if (InSandboxName.IsEmpty())
+	{
+		return EDirectorySuitability::EmptyName;
+	}
+	
+	const FString NewDirectory = FPaths::Combine(InBaseDirectory, InSandboxName);
+	if (!FPaths::ValidatePath(NewDirectory))
+	{
+		return EDirectorySuitability::InvalidPath;
+	}
+
+	if (FindSandboxByName(InSandboxName, InBaseDirectory))
+	{
+		return EDirectorySuitability::DuplicateName;
+	}
+	
+	// We'll require that a sandbox be created in a new directory so we don't have to worry handling pre-existing files in that directory.
+	IPlatformFile& PlatformFile = FPlatformFileManager::Get().GetPlatformFile();
+	const bool bDirectoryAlreadyExists = PlatformFile.DirectoryExists(*NewDirectory);
+	if (bDirectoryAlreadyExists)
+	{
+		return EDirectorySuitability::DirectoryAlreadyExists;
+	}
+	
+	return EDirectorySuitability::Suitable;
+}
+
+FText FormatDirectorySuitabilityAsReason(EDirectorySuitability Suitability, const FString& InSandboxName, const FString& InBaseDirectory)
+{
+	switch (Suitability)
+	{
+	case EDirectorySuitability::Suitable: return LOCTEXT("Suitability.Suitable", "Suitable");
+	case EDirectorySuitability::EmptyName: return LOCTEXT("Suitability.EmptyName", "Sandbox name cannot be empty");
+	case EDirectorySuitability::DuplicateName:
+		return FText::Format(LOCTEXT("Suitability.DuplicateName", "Sandbox name {0} already exists"), FText::FromString(InSandboxName));
+		
+	case EDirectorySuitability::InvalidPath:
+		return FText::Format(
+			LOCTEXT("Suitability.InvalidPath", "Directory {0} is malformed"),
+			FText::FromString(FPaths::Combine(InBaseDirectory, InSandboxName))
+			);
+		
+	case EDirectorySuitability::DirectoryAlreadyExists:
+		return FText::Format(
+			LOCTEXT("Suitability.DirectoryAlreadyExists", "Directory {0} already exists"),
+			FText::FromString(FPaths::Combine(InBaseDirectory, InSandboxName))
+			);
+		
+	default: checkNoEntry(); return FText::GetEmpty();
+	}
+}
+}
+
+#undef LOCTEXT_NAMESPACE
\ No newline at end of file
diff --git a/Engine/Plugins/Developer/Sandbox/FileSandbox/Source/FileSandboxCore/Private/Utils/SandboxUtils.cpp b/Engine/Plugins/Developer/Sandbox/FileSandbox/Source/FileSandboxCore/Private/Utils/SandboxUtils.cpp
new file mode 100644
--- /dev/null
+++ b/Engine/Plugins/Developer/Sandbox/FileSandbox/Source/FileSandboxCore/Private/Utils/SandboxUtils.cpp
@@ -0,0 +1,138 @@
+﻿// Copyright Epic Games, Inc. All Rights Reserved.
+
+#include "SandboxUtils.h"
+
+#include "DirectoryWatcherModule.h"
+#include "IDirectoryWatcher.h"
+#include "Elements/Columns/TypedElementRevisionControlColumns.h"
+#include "HAL/FileManager.h"
+#include "HAL/IConsoleManager.h"
+#include "Modules/ModuleManager.h"
+#include "UObject/Linker.h"
+#include "UObject/LinkerLoad.h"
+#include "UObject/Package.h"
+#include "UObject/UObjectGlobals.h"
+#include "Utils/PackageSandboxUtils.h"
+
+namespace UE::FileSandboxCore
+{
+void ReplayChanges(TConstArrayView<FString> DeletedFiles, TConstArrayView<FString> ModifiedFiles, TConstArrayView<FString> AddedFiles)
+{
+	// Causes any open editors to refresh
+	TArray<FName> PackagesPendingHotReload;
+	for (const FString& Modified : ModifiedFiles)
+	{
+		FString PackageName;
+		if (FPackageName::TryConvertFilenameToLongPackageName(Modified, PackageName))
+		{
+			PackagesPendingHotReload.Add(FName(*PackageName));
+		}
+	}
+	HotReloadPackages(PackagesPendingHotReload);
+
+	// This closes any open editors and removes the files from the content browser
+	TArray<FName> PackagesPendingPurge;
+	IFileManager& FileManager = IFileManager::Get();
+	for (const FString& Deleted : DeletedFiles)
+	{
+		FString PackageName;
+		if (FPackageName::TryConvertFilenameToLongPackageName(Deleted, PackageName))
+		{
+			PackagesPendingPurge.Add(FName(*PackageName));
+			FileManager.Delete(*Deleted);
+		}
+	}
+	PurgePackages(PackagesPendingPurge);
+
+	// This is needed to make added files show up in the content browser
+	TArray<FFileChangeData> FileChanges;
+	Algo::Transform(AddedFiles, FileChanges, [](const FString& File){ return FFileChangeData(File, FFileChangeData::FCA_Added); });
+	if (FDirectoryWatcherModule* DirectoryWatcherModule = GetDirectoryWatcherModuleIfLoaded())
+	{
+		DirectoryWatcherModule->RegisterExternalChanges(FileChanges);
+	}
+
+	SynchronizeAssetRegistry();
+}
+	
+void SynchronizeAssetRegistry()
+{
+#if WITH_EDITOR
+	IDirectoryWatcher* DirectoryWatcher = GetDirectoryWatcherIfLoaded();
+	if (!DirectoryWatcher)
+	{
+		return;
+	}
+
+	DirectoryWatcher->Tick(0.0f);
+#endif // WITH_EDITOR
+}
+
+void FlushPackageLoading(const FString& InPackageName, bool bForceBulkDataLoad)
+{
+	UPackage* ExistingPackage = FindPackage(nullptr, *InPackageName);
+	if (ExistingPackage)
+	{
+		if (!ExistingPackage->IsFullyLoaded())
+		{
+			FlushAsyncLoading();
+			ExistingPackage->FullyLoad();
+		}
+
+		if (bForceBulkDataLoad)
+		{
+			ResetLoaders(ExistingPackage);
+		}
+		else if (ExistingPackage->GetLinker())
+		{
+			ExistingPackage->GetLinker()->Detach();
+		}
+	}
+}
+
+bool FlushPackageFile(const FString& InFilename, FName* OutPackageName, bool bForceLoad)
+{
+	FString PackageName;
+	if (FPackageName::TryConvertFilenameToLongPackageName(InFilename, PackageName))
+	{
+		FlushPackageLoading(PackageName, bForceLoad);
+		if (OutPackageName)
+		{
+			*OutPackageName = *PackageName;
+		}
+		return true;
+	}
+	return false;
+}
+
+FDirectoryWatcherModule& GetDirectoryWatcherModule()
+{
+	static const FName DirectoryWatcherModuleName = TEXT("DirectoryWatcher");
+	return FModuleManager::LoadModuleChecked<FDirectoryWatcherModule>(DirectoryWatcherModuleName);
+}
+
+FDirectoryWatcherModule* GetDirectoryWatcherModuleIfLoaded()
+{
+	static const FName DirectoryWatcherModuleName = TEXT("DirectoryWatcher");
+	if (FModuleManager::Get().IsModuleLoaded(DirectoryWatcherModuleName))
+	{
+		return &FModuleManager::GetModuleChecked<FDirectoryWatcherModule>(DirectoryWatcherModuleName);
+	}
+	return nullptr;
+}
+
+IDirectoryWatcher* GetDirectoryWatcher()
+{
+	FDirectoryWatcherModule& DirectoryWatcherModule = GetDirectoryWatcherModule();
+	return DirectoryWatcherModule.Get();
+}
+
+IDirectoryWatcher* GetDirectoryWatcherIfLoaded()
+{
+	if (FDirectoryWatcherModule* DirectoryWatcherModule = GetDirectoryWatcherModuleIfLoaded())
+	{
+		return DirectoryWatcherModule->Get();
+	}
+	return nullptr;
+}
+}
diff --git a/Engine/Plugins/Developer/Sandbox/FileSandbox/Source/FileSandboxCore/Private/Utils/SandboxUtils.h b/Engine/Plugins/Developer/Sandbox/FileSandbox/Source/FileSandboxCore/Private/Utils/SandboxUtils.h
new file mode 100644
--- /dev/null
+++ b/Engine/Plugins/Developer/Sandbox/FileSandbox/Source/FileSandboxCore/Private/Utils/SandboxUtils.h
@@ -0,0 +1,29 @@
+﻿// Copyright Epic Games, Inc. All Rights Reserved.
+
+#pragma once
+
+#include "Containers/ContainersFwd.h"
+
+class FDirectoryWatcherModule;
+class FName;
+class FString;
+class IDirectoryWatcher;
+
+namespace UE::FileSandboxCore
+{
+/** Utility for applying changes from manifest file to the current IPlatformFile. */
+void ReplayChanges(TConstArrayView<FString> DeletedFiles, TConstArrayView<FString> ModifiedFiles, TConstArrayView<FString> AddedFiles);
+	
+/** Refreshes the assets registry after making underlying file changes. */
+void SynchronizeAssetRegistry();
+
+/** Finishes loading the package, if currently async loading, and detaches the linker. */
+void FlushPackageLoading(const FString& InPackageName, bool bForceBulkDataLoad = true);
+/** Alternate variant that flushes package loading by filename instead. */
+bool FlushPackageFile(const FString& InFilename, FName* OutPackageName = nullptr, bool bForceLoad = true);
+	
+FDirectoryWatcherModule& GetDirectoryWatcherModule();
+FDirectoryWatcherModule* GetDirectoryWatcherModuleIfLoaded();
+IDirectoryWatcher* GetDirectoryWatcher();
+IDirectoryWatcher* GetDirectoryWatcherIfLoaded();
+}
diff --git a/Engine/Plugins/Developer/Sandbox/FileSandbox/Source/FileSandboxCore/Public/Data/InstancedStructMap.h b/Engine/Plugins/Developer/Sandbox/FileSandbox/Source/FileSandboxCore/Public/Data/InstancedStructMap.h
new file mode 100644
--- /dev/null
+++ b/Engine/Plugins/Developer/Sandbox/FileSandbox/Source/FileSandboxCore/Public/Data/InstancedStructMap.h
@@ -0,0 +1,83 @@
+﻿// Copyright Epic Games, Inc. All Rights Reserved.
+
+#pragma once
+
+#include "StructUtils/InstancedStruct.h"
+#include "StructUtils/StructView.h"
+#include "UObject/SoftObjectPath.h"
+#include "InstancedStructMap.generated.h"
+
+/** Advanced tagging structure. Each USTRUCT() can be added at most once. */
+USTRUCT()
+struct FFileSandboxCore_InstancedStructMap
+{
+	GENERATED_BODY()
+
+	/** @return The existing struct data, or newly created data. */
+	FStructView FindOrAddData(UScriptStruct* InStruct);
+
+	/** @return The existing struct data, or an empty view if not present. */
+	FStructView FindData(UScriptStruct* InStruct);
+	
+	/** Removes a struct entry. */
+	void Remove(UScriptStruct* InStruct);
+
+	/** @return The existing struct data, or newly created data. */
+	template<typename T>
+	T& FindOrAddData();
+
+	/** @return The existing struct data, or an empty view if not present. */
+	template<typename T> 
+	T* FindData();
+
+	/** Removes a struct entry. */
+	template<typename T> 
+	void Remove();
+
+private:
+
+	UPROPERTY()
+	TMap<FSoftObjectPath, FInstancedStruct> Data;
+};
+
+inline FStructView FFileSandboxCore_InstancedStructMap::FindOrAddData(UScriptStruct* InStruct)
+{
+	FInstancedStruct& StructData = Data.FindOrAdd(InStruct);
+	if (StructData.GetScriptStruct() == InStruct)
+	{
+		return StructData;
+	}
+
+	StructData.InitializeAs(InStruct);
+	return StructData;
+}
+
+inline FStructView FFileSandboxCore_InstancedStructMap::FindData(UScriptStruct* InStruct)
+{
+	FInstancedStruct* Struct = Data.Find(InStruct);
+	return Struct ? FStructView(*Struct) : FStructView();
+}
+
+inline void FFileSandboxCore_InstancedStructMap::Remove(UScriptStruct* InStruct)
+{
+	Data.Remove(InStruct);
+}
+
+template <typename T>
+T& FFileSandboxCore_InstancedStructMap::FindOrAddData()
+{
+	return FindOrAddData(T::StaticStruct()).template Get<T>();
+}
+
+template <typename T>
+T* FFileSandboxCore_InstancedStructMap::FindData()
+{
+	FStructView Result = FindData(T::StaticStruct());
+	return Result.GetPtr<T>();
+}
+
+template <typename T>
+void FFileSandboxCore_InstancedStructMap::Remove()
+{
+	Remove(T::StaticStruct());
+}
diff --git a/Engine/Plugins/Developer/Sandbox/FileSandbox/Source/FileSandboxCore/Public/Data/SandboxMetaData.h b/Engine/Plugins/Developer/Sandbox/FileSandbox/Source/FileSandboxCore/Public/Data/SandboxMetaData.h
new file mode 100644
--- /dev/null
+++ b/Engine/Plugins/Developer/Sandbox/FileSandbox/Source/FileSandboxCore/Public/Data/SandboxMetaData.h
@@ -0,0 +1,35 @@
+﻿// Copyright Epic Games, Inc. All Rights Reserved.
+
+#pragma once
+#include "InstancedStructMap.h"
+
+#include "SandboxMetaData.generated.h"
+
+/** Metadata about a sandbox. This data is not needed for operations of the sandbox itself and is intended to be useful for external systems. */
+USTRUCT()
+struct FFileSandboxCore_SandboxMetaData
+{
+	GENERATED_BODY()
+
+	/** Name of this sandbox, usually provided by the user. */
+	UPROPERTY()
+	FString Name;
+
+	/** Description of this sandbox, usually provided by the user. */
+	UPROPERTY()
+	FString Description;
+
+	/** Custom tags that may have been added to this sandbox. */
+	UPROPERTY()
+	TArray<FName> Tags;
+
+	/** Custom advanced data that may have been added to this sandbox. */
+	UPROPERTY()
+	FFileSandboxCore_InstancedStructMap CustomData;
+
+	FFileSandboxCore_SandboxMetaData() = default;
+	explicit FFileSandboxCore_SandboxMetaData(FString InName, FString InDescription = {})
+		: Name(MoveTemp(InName))
+		, Description(MoveTemp(InDescription))
+	{}
+};
\ No newline at end of file
diff --git a/Engine/Plugins/Developer/Sandbox/FileSandbox/Source/FileSandboxCore/Public/IFileSandboxCoreModule.h b/Engine/Plugins/Developer/Sandbox/FileSandbox/Source/FileSandboxCore/Public/IFileSandboxCoreModule.h
new file mode 100644
--- /dev/null
+++ b/Engine/Plugins/Developer/Sandbox/FileSandbox/Source/FileSandboxCore/Public/IFileSandboxCoreModule.h
@@ -0,0 +1,37 @@
+// Copyright Epic Games, Inc. All Rights Reserved.
+
+#pragma once
+
+#include "Modules/ModuleInterface.h"
+#include "Modules/ModuleManager.h"
+
+namespace UE::FileSandboxCore
+{
+class ISandboxBrowser;
+class ISandboxInstance;
+class ISandboxManager;
+struct FSandboxInitArgs;
+
+class IFileSandboxCoreModule : public IModuleInterface
+{
+public:
+
+	/**
+	 * Singleton-like access to this module's interface.  This is just for convenience!
+	 * Beware of calling this during the shutdown phase, though.  Your module might have been unloaded already.
+	 *
+	 * @return Returns singleton instance, loading the module on demand if needed
+	 */
+	static IFileSandboxCoreModule& Get() { return FModuleManager::LoadModuleChecked<IFileSandboxCoreModule>("FileSandboxCore"); }
+
+	/**
+	 * Checks to see if this module is loaded and ready.  It is only valid to call Get() during shutdown if IsAvailable() returns true.
+	 *
+	 * @return True if the module is loaded and ready to use
+	 */
+	static bool IsAvailable() { return FModuleManager::Get().IsModuleLoaded("FileSandboxCore"); }
+	
+	/** @return The API for managing live sandboxes. */
+	virtual ISandboxManager& GetSandboxManager() = 0;
+};
+}
diff --git a/Engine/Plugins/Developer/Sandbox/FileSandbox/Source/FileSandboxCore/Public/ISandboxInstance.h b/Engine/Plugins/Developer/Sandbox/FileSandbox/Source/FileSandboxCore/Public/ISandboxInstance.h
new file mode 100644
--- /dev/null
+++ b/Engine/Plugins/Developer/Sandbox/FileSandbox/Source/FileSandboxCore/Public/ISandboxInstance.h
@@ -0,0 +1,79 @@
+﻿// Copyright Epic Games, Inc. All Rights Reserved.
+
+#pragma once
+
+#include "Types/Sandbox/DiscardArgs.h"
+#include "Types/Sandbox/PersistResult.h"
+
+#define UE_API FILESANDBOXCORE_API
+
+struct FFileSandboxCore_SandboxMetaData;
+
+namespace UE::FileSandboxCore
+{
+enum class EBreakBehavior : uint8;
+struct FDiscardResult;
+struct FPersistArgs;
+struct FPersistedFileInfo;
+
+using FConsumeFileSignature = EBreakBehavior(const FString&);
+	
+/**
+ * Manages the file sandbox.
+ *
+ * While this is alive, some engine I/O operations are sandboxed.
+ * Only files in mount points (game content, engine content, plugin content) are tracked.
+ * Other files, such as those residing in Saved/ or Config/ are not affected.
+ */
+class ISandboxInstance
+{
+public:
+	
+	/** Persists all files. */
+	UE_API void PersistAll();
+	
+	/**
+	 * More advanced version of PersistAll that gives you more control how persisting should occur.
+	 * @param InArgs	Parameters controlling the persist options specified by the user.
+	 * @return Returns a persist result object that contains result status.
+	 */
+	virtual FPersistResult PersistSandbox(const FPersistArgs& InArgs) = 0;
+
+
+	
+	/** Discard the sandbox state reverting all files to the state they had before the sandbox. */
+	UE_API void DiscardAll();
+	
+	/**
+	 * More advanced version of DiscardAll that gives you control when packages are hot-reloaded and purged.
+	 * This will also gather packages that need to be hot reloaded or purge from memory.
+	 * 
+	 * @param InArgs Input for the discard operation
+	 * @return Information about the discard operation; you must handle hot-reloading and purging yourself.
+	 */
+	[[nodiscard]] virtual FDiscardResult DiscardSandbox(const FDiscardArgs& InArgs = {}) = 0;
+	
+	
+	/** Enumerates all files that have been changed (added, removed, or edited). */
+	virtual TArray<FString> GatherChangedFiles() const = 0;
+
+	
+	/** @return Whether the given package file exists on non sandbox path, i.e. whether the file was deleted in sandbox mode. */
+	virtual bool DeletedPackageExistsInNonSandbox(const FString& InFilename) const = 0;
+	
+	/**
+	 * @return Metadata about this sandbox.
+	 * 
+	 * @note This is the metadata that the sandbox has when it was created / loaded.
+	 * If the underlying file has changed, this will not return the latest value.
+	 */
+	virtual const FFileSandboxCore_SandboxMetaData& GetInitialMetaData() const = 0;
+
+	/** @return The root directory of the sandbox. */
+	virtual const TCHAR* GetRootDirectory() const = 0;
+
+	virtual ~ISandboxInstance() = default;
+};
+}
+
+#undef UE_API
\ No newline at end of file
diff --git a/Engine/Plugins/Developer/Sandbox/FileSandbox/Source/FileSandboxCore/Public/ISandboxManager.h b/Engine/Plugins/Developer/Sandbox/FileSandbox/Source/FileSandboxCore/Public/ISandboxManager.h
new file mode 100644
--- /dev/null
+++ b/Engine/Plugins/Developer/Sandbox/FileSandbox/Source/FileSandboxCore/Public/ISandboxManager.h
@@ -0,0 +1,125 @@
+﻿// Copyright Epic Games, Inc. All Rights Reserved.
+
+#pragma once
+
+#include "Delegates/Delegate.h"
+#include "Types/Manager/LeaveSandboxArgs.h"
+#include "Types/Manager/SandboxCreationResult.h"
+
+namespace UE::FileSandboxCore
+{
+class FSandboxHandle;
+class ISandboxInstance;
+struct FDeleteSandboxResult;
+struct FLoadSandboxByDirectoryArgs;
+struct FLoadSandboxByNameArgs;
+struct FLoadSandboxError;
+struct FNewSandboxArgs;
+struct FNewSandboxError;
+struct FSandboxCreationResult;
+	
+DECLARE_MULTICAST_DELEGATE_OneParam(FSandboxInstanceEvent, ISandboxInstance&);
+using FSandboxShutdownEvent = FSimpleMulticastDelegate;
+	
+/**
+ * Manages the lifetime of the ISandboxInstance.
+ * For now, only one sandbox can be active at any given time and all files in mount points are sandboxed.
+ *
+ * ===== Overview =====
+ * During sandbox, any I/O operations that modify, delete, or add files in mount points (e.g. /Game/..., /Engine/..., /MyPlugin/..., etc.) will
+ * be sandboxed. That means those operations are performed in a separate sandbox file system without changing the underlying files. At any time,
+ * you can persist changes, i.e. apply them to the underlying file system where the real files reside.
+ *
+ * ==== Example =====
+ * 1. The user opens the asset "Game/Levels/MyBlueprint.uasset", changes it, and then clicks save.
+ * 2. The contained UObjects are serialized and a low-level I/O operation is triggered (via IPlatformFile).
+ *	2.1 Normally, the data would be written "D:/MyProject/Content/Levels/MyBlueprint.uasset",
+ *	2.2 But we redirect the data to a sandboxed location, e.g."D:/MyProject/Intermediate/Sandboxes/YourSandbox/Sandbox/Game".
+ * Similar behaviour goes for when a file is deleted or added.
+ *	
+ * ==== Persisting =====
+ * Persisting effectively means to apply the changes to the original files.
+ * Files marked for delete are deleted from the original location.
+ * Added and modified files are copied from the sandbox location to the original location, e.g. a modified file
+ * "D:/MyProject/Content/Levels/MyBlueprint.uasset" would be replaced with "D:/MyProject/Intermediate/Sandboxes/YourSandbox/Sandbox/Game".
+ *
+ * ===== Sandbox file structure =====
+ * By default, sandboxes are saved to the Intermediate/Sandboxes directory (@see GetBaseSandboxDirectory).
+ * Sandboxes are stored in separate directories, e.g. a sandbox named "MySandbox" would be saved in "MyProject/Intermediate/Sandboxes/MySandbox".
+ * 
+ * The directory contains a "Manifest.json" file at the root level; it contains bookkeeping data (which files were added, modified, deleted) and
+ * metadata, such as user given description and custom metadata specific to the business logic.
+ * So for example, there would be a file "MyProject/Intermediate/Sandboxes/MySandbox/Manifest.json".
+ * 
+ * Furthermore, changed and added files are saved in the sub-directory "Sandbox". This directory contains one directory name per mount point,
+ * like "Game", "Engine", "MyPlugin". Within those mount point directories files are added according to the folder structure of the package name.
+ * So for example:
+ * - "MyProject/Intermediate/Sandboxes/MySandbox/Sandboxes/Game/Levels/MyLevel.uasset" maps to "Game/Levels/MyLevel.uasset" (mount = "/Game/"
+ * - "MyProject/Intermediate/Sandboxes/MySandbox/Sandboxes/MyPlugin/FooBlueprint.uasset" maps to "MyPlugin/FooBlueprint.uasset" (mount = "/MyPlugin/")
+ */
+class ISandboxManager
+{
+public:
+	
+	/**
+	 * Puts the engine into a sandbox: This creates a new environment.
+	 * 
+	 * Tip: If you don't care about the specifics of potential errors, you can write if (ISandboxInstance* Instance = Manager->CreateNewSandbox(...)).
+	 * @note Do not keep any raw reference to ISandboxInstance as it turns stale when the sandbox is shutdown.
+	 * @see FWeakSandboxPtr, GetGlobalSandboxInstance.
+	 */
+	virtual FNewSandboxResult CreateNewSandbox(const FNewSandboxArgs& InArgs) = 0;
+
+	
+	/**
+	 * Puts the engine into a sandbox: applies the state from a pre-existing sandbox to any open assets.
+	 * 
+	 * Any assets that are currently open / in-memory will be hot-reloaded.
+	 * Effectively, this means that their file content will be re-loaded but from the sandboxed version of the .uasset file.
+	 * 
+	 * Tip: If you don't care about the specifics of potential errors, you can write if (ISandboxInstance* Instance = Manager->LoadNamedSandbox(...)).
+	 * @note Do not keep any raw reference to ISandboxInstance as it turns stale when the sandbox is shutdown.
+	 * @see FWeakSandboxPtr, GetGlobalSandboxInstance.
+	 */
+	FLoadSandboxResult LoadNamedSandbox(const FLoadSandboxByNameArgs& InName);
+	
+	/** Advanced version of LoadNamedSandbox for when you know the specific root directory of a sandbox. */
+	virtual FLoadSandboxResult LoadSandbox(const FLoadSandboxByDirectoryArgs& InLoadArgs) = 0;
+	
+	
+	/** Deletes the sandbox with the given name, if found. It is valid to delete a sandbox if it is currently valid. */
+	FDeleteSandboxResult DeleteNamedSandbox(const FString& InSandboxName);
+	
+	/** Advanced version of DeleteNamedSandbox for when you know the specific root directory of a sandbox. */
+	virtual FDeleteSandboxResult DeleteSandbox(const FString& InSandboxDirectory) = 0;
+	
+	
+	/** Leaves the current sandbox. */
+	virtual void LeaveSandbox(const FLeaveSandboxArgs& InLeaveArgs = FLeaveSandboxArgs()) = 0;
+
+	/**
+	 * Gets the currently active global sandbox instance.
+	 *
+	 * The name “global” sandbox instance was chosen with the idea that, in the future, we may support directory bindings where only specific
+	 * directories are sandboxed. For now (5.8), there can be at most one sandbox instance.
+	 *
+	 * @note Do not keep any raw reference to ISandboxInstance as it turns stale when the sandbox is shutdown.
+	 * @see FWeakSandboxPtr, GetGlobalSandboxInstance.
+	 */
+	virtual ISandboxInstance* GetActiveSandboxInstance() const = 0;
+
+	/** @return Whether the is any active sandbox. */
+	bool HasActiveSandbox() const { return GetActiveSandboxInstance() != nullptr; }
+
+	
+	/** @return Event invoked after a sandbox has been successfully created. */
+	virtual FSandboxInstanceEvent& OnPostSandboxStartup() = 0;
+	
+	/** @return Event invoked when it has been decided, that a sandbox instance will be taken down, but before any core destruction logic runs. */
+	virtual FSandboxInstanceEvent& OnPreSandboxShutdown() = 0;
+	/** @return Event invoked at the end of shutting down a sandbox, i.e. after all core destruction logic has run. */
+	virtual FSandboxShutdownEvent& OnPostSandboxShutdown() = 0;
+	
+	virtual ~ISandboxManager() = default;
+};
+}
diff --git a/Engine/Plugins/Developer/Sandbox/FileSandbox/Source/FileSandboxCore/Public/Types/EBreakBehavior.h b/Engine/Plugins/Developer/Sandbox/FileSandbox/Source/FileSandboxCore/Public/Types/EBreakBehavior.h
new file mode 100644
--- /dev/null
+++ b/Engine/Plugins/Developer/Sandbox/FileSandbox/Source/FileSandboxCore/Public/Types/EBreakBehavior.h
@@ -0,0 +1,14 @@
+﻿// Copyright Epic Games, Inc. All Rights Reserved.
+
+#pragma once
+
+#include "HAL/Platform.h"
+
+namespace UE::FileSandboxCore
+{
+enum class EBreakBehavior : uint8
+{
+	Continue,
+	Break
+};
+}
\ No newline at end of file
diff --git a/Engine/Plugins/Developer/Sandbox/FileSandbox/Source/FileSandboxCore/Public/Types/Manager/DeleteSandboxResult.h b/Engine/Plugins/Developer/Sandbox/FileSandbox/Source/FileSandboxCore/Public/Types/Manager/DeleteSandboxResult.h
new file mode 100644
--- /dev/null
+++ b/Engine/Plugins/Developer/Sandbox/FileSandbox/Source/FileSandboxCore/Public/Types/Manager/DeleteSandboxResult.h
@@ -0,0 +1,36 @@
+﻿// Copyright Epic Games, Inc. All Rights Reserved.
+
+#pragma once
+
+#include "HAL/Platform.h"
+
+namespace UE::FileSandboxCore
+{
+enum class EDeleteSandboxErrorCode : uint8
+{
+	/** Sandbox was deleted successfully. */
+	Success,
+
+	/** Given sandbox name could not be resolved to any sandbox. */
+	InvalidName,
+
+	/** Given directory does not contain any sandbox. */
+	InvalidDirectory,
+
+	/** An I/O error occured, e.g. the directory could not be deleted. */
+	IOError
+};
+
+/** Describes the result of the delete operation, with a reason of why it failed. */
+struct FDeleteSandboxResult
+{
+	/** Describes the result of the delete operation, with a reason of why it failed. */
+	EDeleteSandboxErrorCode ErrorCode = EDeleteSandboxErrorCode::Success;
+
+	FDeleteSandboxResult() = default;
+	FDeleteSandboxResult(EDeleteSandboxErrorCode InErrorCode) : ErrorCode(InErrorCode) {}
+
+	/** @return Whether the sandbox was successfully deleted. */
+	bool IsSuccess() const { return ErrorCode == EDeleteSandboxErrorCode::Success; }
+};
+}
\ No newline at end of file
diff --git a/Engine/Plugins/Developer/Sandbox/FileSandbox/Source/FileSandboxCore/Public/Types/Manager/LeaveSandboxArgs.h b/Engine/Plugins/Developer/Sandbox/FileSandbox/Source/FileSandboxCore/Public/Types/Manager/LeaveSandboxArgs.h
new file mode 100644
--- /dev/null
+++ b/Engine/Plugins/Developer/Sandbox/FileSandbox/Source/FileSandboxCore/Public/Types/Manager/LeaveSandboxArgs.h
@@ -0,0 +1,16 @@
+﻿// Copyright Epic Games, Inc. All Rights Reserved.
+
+#pragma once
+
+#include "Misc/Optional.h"
+#include "Types/Sandbox/PersistArgs.h"
+
+namespace UE::FileSandboxCore
+{
+/** Describes how the session is supposed to be left. */
+struct FLeaveSandboxArgs
+{
+	/** If set, this specifies the files that are supposed to be persisted, and how. */
+	TOptional<FPersistArgs> PersistArgs;
+};
+}
diff --git a/Engine/Plugins/Developer/Sandbox/FileSandbox/Source/FileSandboxCore/Public/Types/Manager/LoadSandboxByDirectoryArgs.h b/Engine/Plugins/Developer/Sandbox/FileSandbox/Source/FileSandboxCore/Public/Types/Manager/LoadSandboxByDirectoryArgs.h
new file mode 100644
--- /dev/null
+++ b/Engine/Plugins/Developer/Sandbox/FileSandbox/Source/FileSandboxCore/Public/Types/Manager/LoadSandboxByDirectoryArgs.h
@@ -0,0 +1,24 @@
+﻿// Copyright Epic Games, Inc. All Rights Reserved.
+
+#pragma once
+
+#include "Containers/UnrealString.h"
+#include "SandboxInitArgs.h"
+
+namespace UE::FileSandboxCore
+{
+/** Args for loading a pre-existing sandbox. */
+struct FLoadSandboxByDirectoryArgs
+{
+	/** The root directory of the sandbox (i.e. the one containing the manifest file). */
+	FString SandboxRootPath;
+	
+	/** Base arguments for instantiating the sandbox environment. */
+	FSandboxInitArgs InitArgs;
+
+	explicit FLoadSandboxByDirectoryArgs(FString InSandboxRootPath) : SandboxRootPath(MoveTemp(InSandboxRootPath))
+	{
+		check(!SandboxRootPath.IsEmpty());
+	}
+};
+}
diff --git a/Engine/Plugins/Developer/Sandbox/FileSandbox/Source/FileSandboxCore/Public/Types/Manager/LoadSandboxByNameArgs.h b/Engine/Plugins/Developer/Sandbox/FileSandbox/Source/FileSandboxCore/Public/Types/Manager/LoadSandboxByNameArgs.h
new file mode 100644
--- /dev/null
+++ b/Engine/Plugins/Developer/Sandbox/FileSandbox/Source/FileSandboxCore/Public/Types/Manager/LoadSandboxByNameArgs.h
@@ -0,0 +1,27 @@
+﻿// Copyright Epic Games, Inc. All Rights Reserved.
+
+#pragma once
+
+#include "SandboxInitArgs.h"
+#include "Utils/SandboxDirectoryUtils.h"
+
+namespace UE::FileSandboxCore
+{
+/** Args for loading a pre-existing sandbox. */
+struct FLoadSandboxByNameArgs
+{
+	/** The name of the sandbox you want to load. */
+	FString Name;
+
+	/** The directory in which to search for the sandbox. */
+	FString BaseDirectory = GetBaseSandboxDirectory();
+	
+	/** Base arguments for instantiating the sandbox environment. */
+	FSandboxInitArgs BaseArgs;
+
+	explicit FLoadSandboxByNameArgs(FString InName) : Name(MoveTemp(InName))
+	{
+		check(!Name.IsEmpty());
+	}
+};
+}
diff --git a/Engine/Plugins/Developer/Sandbox/FileSandbox/Source/FileSandboxCore/Public/Types/Manager/LoadSandboxError.h b/Engine/Plugins/Developer/Sandbox/FileSandbox/Source/FileSandboxCore/Public/Types/Manager/LoadSandboxError.h
new file mode 100644
--- /dev/null
+++ b/Engine/Plugins/Developer/Sandbox/FileSandbox/Source/FileSandboxCore/Public/Types/Manager/LoadSandboxError.h
@@ -0,0 +1,28 @@
+﻿// Copyright Epic Games, Inc. All Rights Reserved.
+
+#pragma once
+
+#include "HAL/Platform.h"
+
+namespace UE::FileSandboxCore
+{
+enum class ELoadSandboxLoadErrorReason : uint8
+{
+	/** The directory does not contain any sandbox */
+	InvalidDirectory,
+	/** The name does not map to any directory. */
+	InvalidName,
+	
+	/** Some generic IO error occured (e.g. could not write manifest file, etc.) */
+	IOError,
+
+	/** Dummy error code for initializing to a deterministic value. */
+	Unspecified
+};
+
+/** Explains why a new sandbox could not be loaded. */ 
+struct FLoadSandboxError
+{
+	ELoadSandboxLoadErrorReason Reason = ELoadSandboxLoadErrorReason::Unspecified;
+};
+}
\ No newline at end of file
diff --git a/Engine/Plugins/Developer/Sandbox/FileSandbox/Source/FileSandboxCore/Public/Types/Manager/NewSandboxArgs.h b/Engine/Plugins/Developer/Sandbox/FileSandbox/Source/FileSandboxCore/Public/Types/Manager/NewSandboxArgs.h
new file mode 100644
--- /dev/null
+++ b/Engine/Plugins/Developer/Sandbox/FileSandbox/Source/FileSandboxCore/Public/Types/Manager/NewSandboxArgs.h
@@ -0,0 +1,39 @@
+﻿// Copyright Epic Games, Inc. All Rights Reserved.
+
+#pragma once
+
+#include "Data/SandboxMetaData.h"
+#include "SandboxInitArgs.h"
+#include "Utils/SandboxDirectoryUtils.h"
+
+namespace UE::FileSandboxCore
+{
+/** Args for creating a new sandbox. */
+struct FNewSandboxArgs
+{
+	/** Directory in which sandboxes live. A new directory with the sandbox name will be created in this directory. */
+	FString SandboxBasePath = GetBaseSandboxDirectory();
+
+	/**
+	 * Metadata to set about this new sandbox.
+	 * 
+	 * The name will be used to create a base directory in the sandbox folder, so it must be non-empty and unique;
+	 * invalid name leads to failure.
+	 */
+	FFileSandboxCore_SandboxMetaData MetaData;
+	
+	/** Base arguments for instantiating the sandbox environment. */
+	FSandboxInitArgs InitArgs;
+
+	explicit FNewSandboxArgs(FString InName, FString InDescription = {})
+		: MetaData(MoveTemp(InName), MoveTemp(InDescription))
+	{
+		check(!MetaData.Name.IsEmpty());
+	}
+	
+	explicit FNewSandboxArgs(FFileSandboxCore_SandboxMetaData MetaData) : MetaData(MoveTemp(MetaData))
+	{
+		check(!MetaData.Name.IsEmpty());
+	}
+};
+}
diff --git a/Engine/Plugins/Developer/Sandbox/FileSandbox/Source/FileSandboxCore/Public/Types/Manager/NewSandboxError.h b/Engine/Plugins/Developer/Sandbox/FileSandbox/Source/FileSandboxCore/Public/Types/Manager/NewSandboxError.h
new file mode 100644
--- /dev/null
+++ b/Engine/Plugins/Developer/Sandbox/FileSandbox/Source/FileSandboxCore/Public/Types/Manager/NewSandboxError.h
@@ -0,0 +1,26 @@
+﻿// Copyright Epic Games, Inc. All Rights Reserved.
+
+#pragma once
+
+#include "HAL/Platform.h"
+
+namespace UE::FileSandboxCore
+{
+enum class ENewSandboxErrorReason : uint8
+{
+	/** The specified directory cannot be used as root directory. */
+	UnsuitablePath,
+
+	/** Some generic IO error occured (e.g. could not write manifest file, etc.) */
+	IOError,
+
+	/** Dummy error code for initializing to a deterministic value. */
+	Unspecified
+};
+
+/** Explains why a new sandbox could not be created. */ 
+struct FNewSandboxError
+{
+	ENewSandboxErrorReason Reason = ENewSandboxErrorReason::Unspecified;
+};
+}
\ No newline at end of file
diff --git a/Engine/Plugins/Developer/Sandbox/FileSandbox/Source/FileSandboxCore/Public/Types/Manager/SandboxCreationResult.h b/Engine/Plugins/Developer/Sandbox/FileSandbox/Source/FileSandboxCore/Public/Types/Manager/SandboxCreationResult.h
new file mode 100644
--- /dev/null
+++ b/Engine/Plugins/Developer/Sandbox/FileSandbox/Source/FileSandboxCore/Public/Types/Manager/SandboxCreationResult.h
@@ -0,0 +1,31 @@
+﻿// Copyright Epic Games, Inc. All Rights Reserved.
+
+#pragma once
+
+#include "LoadSandboxError.h"
+#include "Misc/NotNull.h"
+#include "NewSandboxError.h"
+#include "Templates/ValueOrError.h"
+
+namespace UE::FileSandboxCore
+{
+class ISandboxInstance;
+
+/**
+ * Utility type to allow implicit conversion to ISandboxInstance*, i.e. to allow shorthand semantics such as
+ *		ISandboxManager* Manager = ...;
+ *		if (ISandboxInstance* SandboxInstance = Manager->LoadNamedSandbox(...)) { ... }
+ */
+template<typename TErrorType>
+struct TSandboxCreationResult : TValueOrError<TNotNull<ISandboxInstance*>, TErrorType>
+{
+	using Super = TValueOrError<TNotNull<ISandboxInstance*>, TErrorType>;
+	
+	TSandboxCreationResult(TValueOrError<TNotNull<ISandboxInstance*>, TErrorType> Result) : Super(MoveTemp(Result)) {}
+
+	operator ISandboxInstance*() const { return this->HasValue() ? this->GetValue() : nullptr; }
+};
+
+using FNewSandboxResult = TSandboxCreationResult<FNewSandboxError>;
+using FLoadSandboxResult = TSandboxCreationResult<FLoadSandboxError>;
+}
diff --git a/Engine/Plugins/Developer/Sandbox/FileSandbox/Source/FileSandboxCore/Public/Types/Manager/SandboxInitArgs.h b/Engine/Plugins/Developer/Sandbox/FileSandbox/Source/FileSandboxCore/Public/Types/Manager/SandboxInitArgs.h
new file mode 100644
--- /dev/null
+++ b/Engine/Plugins/Developer/Sandbox/FileSandbox/Source/FileSandboxCore/Public/Types/Manager/SandboxInitArgs.h
@@ -0,0 +1,22 @@
+﻿// Copyright Epic Games, Inc. All Rights Reserved.
+
+#pragma once
+
+#include "HAL/PlatformFileManager.h"
+#include "Misc/NotNull.h"
+
+namespace UE::FileSandboxCore
+{
+/** Arguments for creating a sandbox. */
+struct FSandboxInitArgs
+{
+	/**
+	 * System for saving out the sandboxed files.
+	 * In theory, you could use this to save out the sandboxed files outside of the filesystem, e.g. to the cloud.
+	 */
+	IPlatformFile* OverridePlatformFile = nullptr;
+
+	/** @return The platform file to use. */
+	IPlatformFile& GetPlatformFile() const { return OverridePlatformFile ? *OverridePlatformFile : FPlatformFileManager::Get().GetPlatformFile(); }
+};
+}
diff --git a/Engine/Plugins/Developer/Sandbox/FileSandbox/Source/FileSandboxCore/Public/Types/PersistedFileInfo.h b/Engine/Plugins/Developer/Sandbox/FileSandbox/Source/FileSandboxCore/Public/Types/PersistedFileInfo.h
new file mode 100644
--- /dev/null
+++ b/Engine/Plugins/Developer/Sandbox/FileSandbox/Source/FileSandboxCore/Public/Types/PersistedFileInfo.h
@@ -0,0 +1,18 @@
+﻿// Copyright Epic Games, Inc. All Rights Reserved.
+
+#pragma once
+
+#include "Containers/UnrealString.h"
+#include "Misc/DateTime.h"
+
+namespace UE::FileSandboxCore
+{
+struct FPersistedFileInfo
+{
+	/** File in the original file system (e.g. "D:/MyProject/Content/Levels/Map.umap"). */
+	FString FileName;
+
+	/** Timestamp of when the file was persisted (to detect if the file has been changed again after it was persisted). */
+	FDateTime Timestamp;
+};
+}
diff --git a/Engine/Plugins/Developer/Sandbox/FileSandbox/Source/FileSandboxCore/Public/Types/Sandbox/DiscardArgs.h b/Engine/Plugins/Developer/Sandbox/FileSandbox/Source/FileSandboxCore/Public/Types/Sandbox/DiscardArgs.h
new file mode 100644
--- /dev/null
+++ b/Engine/Plugins/Developer/Sandbox/FileSandbox/Source/FileSandboxCore/Public/Types/Sandbox/DiscardArgs.h
@@ -0,0 +1,12 @@
+﻿// Copyright Epic Games, Inc. All Rights Reserved.
+
+#pragma once
+
+namespace UE::FileSandboxCore
+{
+/** Input arguments for discarding a sandbox. */
+struct FDiscardArgs
+{
+	// TODO DP: In the future, add additional args here, such as callbacks to implement a progress bar
+};
+}
diff --git a/Engine/Plugins/Developer/Sandbox/FileSandbox/Source/FileSandboxCore/Public/Types/Sandbox/DiscardResult.h b/Engine/Plugins/Developer/Sandbox/FileSandbox/Source/FileSandboxCore/Public/Types/Sandbox/DiscardResult.h
new file mode 100644
--- /dev/null
+++ b/Engine/Plugins/Developer/Sandbox/FileSandbox/Source/FileSandboxCore/Public/Types/Sandbox/DiscardResult.h
@@ -0,0 +1,25 @@
+﻿// Copyright Epic Games, Inc. All Rights Reserved.
+
+#pragma once
+
+#include "Containers/Array.h"
+#include "UObject/NameTypes.h"
+
+namespace UE::FileSandboxCore
+{
+/** Result of discarding a sandbox. */
+struct FDiscardResult
+{
+	/**
+	 * The package names that should be hot reloaded.
+	 * You must call HotReloadPackages on them.
+	 */
+	TArray<FName> PackagesPendingHotReload;
+	
+	/**
+	 * The package names that should be purged.
+	 * You must call PurgePackages on them.
+	 */
+	TArray<FName> PackagesPendingPurge;
+};
+}
diff --git a/Engine/Plugins/Developer/Sandbox/FileSandbox/Source/FileSandboxCore/Public/Types/Sandbox/PersistArgs.h b/Engine/Plugins/Developer/Sandbox/FileSandbox/Source/FileSandboxCore/Public/Types/Sandbox/PersistArgs.h
new file mode 100644
--- /dev/null
+++ b/Engine/Plugins/Developer/Sandbox/FileSandbox/Source/FileSandboxCore/Public/Types/Sandbox/PersistArgs.h
@@ -0,0 +1,21 @@
+﻿// Copyright Epic Games, Inc. All Rights Reserved.
+
+#pragma once
+
+#include "Containers/ArrayView.h"
+#include "Containers/UnrealString.h"
+
+class FString;
+
+namespace UE::FileSandboxCore
+{
+/** Input arguments for persisting a sandbox. */
+struct FPersistArgs
+{
+	/** List of files names from the original filesystem, e.g. "D:/MyProject/Content/Levels/MyLevel.umap" */
+	TConstArrayView<FString> Files;
+	
+	/** Whether to force the file writable when no source control provider is active. */
+	bool bShouldMakeWritableIfNoSourceControl = false;
+};
+}
diff --git a/Engine/Plugins/Developer/Sandbox/FileSandbox/Source/FileSandboxCore/Public/Types/Sandbox/PersistResult.h b/Engine/Plugins/Developer/Sandbox/FileSandbox/Source/FileSandboxCore/Public/Types/Sandbox/PersistResult.h
new file mode 100644
--- /dev/null
+++ b/Engine/Plugins/Developer/Sandbox/FileSandbox/Source/FileSandboxCore/Public/Types/Sandbox/PersistResult.h
@@ -0,0 +1,33 @@
+﻿// Copyright Epic Games, Inc. All Rights Reserved.
+
+#pragma once
+
+#include "Containers/ArrayView.h"
+#include "Internationalization/Text.h"
+#include "HAL/Platform.h"
+
+class ISourceControlProvider;
+
+namespace UE::FileSandboxCore
+{
+/** Possible outcomes from a session persist. */
+enum class EPersistStatus : uint8
+{
+	/** The packages were persisted correctly. */
+	Success,
+	/** We failed to persist the packages. */
+	Failure,
+	/** We are not allowed to persist. */
+	NotAllowed
+};
+
+/** Result of a persist.*/
+struct FPersistResult
+{
+	/** List of reasons for persist failure. */
+	TArray<FText> FailureReasons;
+
+	/** Success or failure of the session persist. */
+	EPersistStatus PersistStatus = EPersistStatus::NotAllowed;
+};
+}
diff --git a/Engine/Plugins/Developer/Sandbox/FileSandbox/Source/FileSandboxCore/Public/Utils/PackageSandboxUtils.h b/Engine/Plugins/Developer/Sandbox/FileSandbox/Source/FileSandboxCore/Public/Utils/PackageSandboxUtils.h
new file mode 100644
--- /dev/null
+++ b/Engine/Plugins/Developer/Sandbox/FileSandbox/Source/FileSandboxCore/Public/Utils/PackageSandboxUtils.h
@@ -0,0 +1,14 @@
+﻿// Copyright Epic Games, Inc. All Rights Reserved.
+
+#pragma once
+
+#include "Containers/Array.h"
+
+class FName;
+
+namespace UE::FileSandboxCore
+{
+void PurgePackages(const TArray<FName>& InPackageNames);
+
+void HotReloadPackages(const TArray<FName>& InPackageNames);
+}
diff --git a/Engine/Plugins/Developer/Sandbox/FileSandbox/Source/FileSandboxCore/Public/Utils/SandboxDirectoryUtils.h b/Engine/Plugins/Developer/Sandbox/FileSandbox/Source/FileSandboxCore/Public/Utils/SandboxDirectoryUtils.h
new file mode 100644
--- /dev/null
+++ b/Engine/Plugins/Developer/Sandbox/FileSandbox/Source/FileSandboxCore/Public/Utils/SandboxDirectoryUtils.h
@@ -0,0 +1,111 @@
+﻿// Copyright Epic Games, Inc. All Rights Reserved.
+
+#pragma once
+
+#include "Containers/UnrealString.h"
+#include "Misc/Optional.h"
+#include "Templates/FunctionFwd.h"
+
+#define UE_API FILESANDBOXCORE_API
+
+class FText;
+
+struct FFileSandboxCore_SandboxMetaData;
+
+namespace UE::FileSandboxCore
+{
+enum class EBreakBehavior : uint8;
+using FProcessSandboxDirectorySignature = EBreakBehavior(const FString& InBaseDirectory);
+
+
+	
+/** @return The base directory into which new sandboxes are stored. */
+UE_API FString GetBaseSandboxDirectory();
+
+/**
+ * Enumerates all directories that contain a sandbox. The search only considers the top-level directories.
+ * @param InProcessSandbox Invokes for each found sandbox.
+ * @param InBaseDirectory Base directory search for. Either absolute or relative to the process BaseDir().
+ */
+UE_API void ForEachSandbox(
+	TFunctionRef<FProcessSandboxDirectorySignature> InProcessSandbox,
+	const FString& InBaseDirectory = GetBaseSandboxDirectory()
+	);
+
+
+	
+enum class EMetaDataWriteFlags : uint8
+{
+	None,
+
+	Name = 1 << 0,
+	Description = 1 << 1,
+	Tags = 1 << 2,
+	CustomData = 1 << 3,
+
+	All = Name | Description | Tags | CustomData
+};
+ENUM_CLASS_FLAGS(EMetaDataWriteFlags);
+	
+/**
+ * Writes new metadata for a sandbox.
+ * 
+ * You're allowed to change the sandbox name. The name of the base directory will remain be changed as well (for consistency).
+ * As such, this function will fail if you're trying to rename the sandbox while it is active (for safety & easier implementation reasons).
+ * 
+ * @param NewMetaData The new metadata to set for the sandbox.
+ * @param InSandboxDirectory Path to the base directory of the sandbox, e.g. what's returned by ForEachSandbox. Either absolute or relative to the process BaseDir().
+ * @param InWriteFlags Optional. A property mask of properties you want to write. If you only wanted to write the Description, you can leave the other members blank.
+ */
+UE_API bool SaveMetaData(
+	const FFileSandboxCore_SandboxMetaData& NewMetaData, const FString& InSandboxDirectory, 
+	EMetaDataWriteFlags InWriteFlags = EMetaDataWriteFlags::All
+	);
+	
+/**
+ * Reads metadata about a sandbox.
+ * @param InSandboxDirectory Path to the base directory of the sandbox, e.g. what's returned by ForEachSandbox. Either absolute or relative to the process BaseDir().
+ */
+UE_API TOptional<FFileSandboxCore_SandboxMetaData> LoadMetaData(const FString& InSandboxDirectory);
+
+	
+	
+/**
+ * Searches for sandboxes with the given name. The search only considers the top-level directories.
+ * 
+ * @param InSandboxName The name of the sandbox
+ * @param InBaseDirectory Base directory search for. Either absolute or relative to the process BaseDir().
+ * @return The path to the directory containing the sandbox with the specified name.
+ */
+UE_API TOptional<FString> FindSandboxByName(
+	const FString& InSandboxName,
+	const FString& InBaseDirectory = GetBaseSandboxDirectory()
+	);
+
+
+	
+enum class EDirectorySuitability : uint8
+{
+	Suitable,
+	EmptyName,
+	DuplicateName,
+	InvalidPath,
+	DirectoryAlreadyExists
+};
+	
+/** @return Whether a sandbox directory named InSandboxName can be put into the InBaseDirectory. */
+UE_API EDirectorySuitability DetermineDirectorySuitability(
+	const FString& InSandboxName,
+	const FString& InBaseDirectory = GetBaseSandboxDirectory(),
+	FText* OutReason = nullptr
+	);
+
+/** @return EDirectorySuitability explained as reason to the user. */
+UE_API FText FormatDirectorySuitabilityAsReason(
+	EDirectorySuitability Suitability,
+	const FString& InSandboxName,
+	const FString& InBaseDirectory = GetBaseSandboxDirectory()
+	);
+}
+
+#undef UE_API
\ No newline at end of file
