File Sandbox: Fix adding file and deleting it causing it to be marked as deleted

#Synced-CL 47637010

diff --git a/Engine/Plugins/Developer/Sandbox/FileSandbox/Source/FileSandboxCore/Private/Sandbox/Platform/SandboxPlatformFile.cpp b/Engine/Plugins/Developer/Sandbox/FileSandbox/Source/FileSandboxCore/Private/Sandbox/Platform/SandboxPlatformFile.cpp
--- a/Engine/Plugins/Developer/Sandbox/FileSandbox/Source/FileSandboxCore/Private/Sandbox/Platform/SandboxPlatformFile.cpp
+++ b/Engine/Plugins/Developer/Sandbox/FileSandbox/Source/FileSandboxCore/Private/Sandbox/Platform/SandboxPlatformFile.cpp
@@ -190,19 +190,29 @@ bool FSandboxPlatformFile::DeleteFile(const TCHAR* Filename)
 
 	if (ResolvedPath.HasSandboxPath())
 	{
+		// Cannot delete again.
 		if (IsPathDeleted(ResolvedPath))
 		{
 			return false;
 		}
 
-		if (!LowerLevel->FileExists(*ResolvedPath.GetSandboxPath()) || LowerLevel->DeleteFile(*ResolvedPath.GetSandboxPath()))
+		// Is file present outside of sandbox? E.g. existed before we entered sandbox. If so, mark as deleted.
+		const bool bIsDeletingNonSandboxFile = !LowerLevel->FileExists(*ResolvedPath.GetSandboxPath());
+		if (bIsDeletingNonSandboxFile)
+		{
+			MarkPathDeleted(ResolvedPath);
+		}
+		
+		// File only present in sandbox? E.g. file created during sandbox, but then deleted again.
+		const bool bIsDeletingFileCreatedInSandbox = LowerLevel->DeleteFile(*ResolvedPath.GetSandboxPath()); 
+		
+		const bool bPerformedDeleteOp = bIsDeletingNonSandboxFile || bIsDeletingFileCreatedInSandbox;
+		if (bPerformedDeleteOp)
 		{
-			SetPathDeleted(ResolvedPath, true);
 			NotifyFileDeleted(ResolvedPath);
-			return true;
 		}
-
-		return false; // Do not attempt to delete the non-sandbox file
+		
+		return bPerformedDeleteOp; 
 	}
 
 	return LowerLevel->DeleteFile(*ResolvedPath.GetNonSandboxPath());
@@ -290,8 +300,8 @@ bool FSandboxPlatformFile::MoveFile(const TCHAR* To, const TCHAR* From)
 
 			if (bSuccess)
 			{
-				SetPathDeleted(ResolvedToPath, false);
-				SetPathDeleted(ResolvedFromPath, true);
+				ClearPathDeleted(ResolvedToPath);
+				MarkPathDeleted(ResolvedFromPath);
 				NotifyFileDeleted(ResolvedFromPath);
 			}
 
@@ -302,7 +312,7 @@ bool FSandboxPlatformFile::MoveFile(const TCHAR* To, const TCHAR* From)
 			// Non-sandbox -> Sandbox
 			if (LowerLevel->MoveFile(*ResolvedToPath.GetSandboxPath(), *ResolvedFromPath.GetNonSandboxPath()))
 			{
-				SetPathDeleted(ResolvedToPath, false);
+				ClearPathDeleted(ResolvedToPath);
 				return true;
 			}
 			return false;
@@ -331,7 +341,7 @@ bool FSandboxPlatformFile::MoveFile(const TCHAR* To, const TCHAR* From)
 
 		if (bSuccess)
 		{
-			SetPathDeleted(ResolvedFromPath, true);
+			MarkPathDeleted(ResolvedFromPath);
 			NotifyFileDeleted(ResolvedFromPath);
 		}
 
@@ -471,7 +481,7 @@ IFileHandle* FSandboxPlatformFile::OpenWrite(const TCHAR* Filename, bool bAppend
 		IFileHandle* Handle = LowerLevel->OpenWrite(*ResolvedPath.GetSandboxPath(), bAppend, bAllowRead);
 		if (Handle)
 		{
-			SetPathDeleted(ResolvedPath, false);
+			ClearPathDeleted(ResolvedPath);
 		}
 		return Handle;
 	}
@@ -507,7 +517,7 @@ bool FSandboxPlatformFile::CreateDirectory(const TCHAR* Directory)
 	{
 		if (LowerLevel->CreateDirectory(*ResolvedPath.GetSandboxPath()))
 		{
-			SetPathDeleted(ResolvedPath, false);
+			ClearPathDeleted(ResolvedPath);
 			return true;
 		}
 
@@ -532,7 +542,7 @@ bool FSandboxPlatformFile::DeleteDirectory(const TCHAR* Directory)
 		const TArray<FDirectoryItem> DirectoryItems = GetDirectoryContents(ResolvedPath, Directory);
 		if (DirectoryItems.Num() == 0 && (!LowerLevel->DirectoryExists(*ResolvedPath.GetSandboxPath()) || LowerLevel->DeleteDirectory(*ResolvedPath.GetSandboxPath())))
 		{
-			SetPathDeleted(ResolvedPath, true);
+			MarkPathDeleted(ResolvedPath);
 			return true;
 		}
 
@@ -857,17 +867,19 @@ void FSandboxPlatformFile::DiscardSandbox(
 		FScopeLock DeletedSandboxPathsLock(&DeletedSandboxPathsCS);
 		for (auto It = DeletedSandboxPaths.CreateIterator(); It; ++It)
 		{
+			const FNonSandboxPath& NonSandboxPath = *It;
+			
 			// If this file maps to a package then we need to flush its linker so that we can remove the file from the sandbox
 			FName PackageName;
-			FlushPackageFile(It->GetNonSandboxPath(), &PackageName);
+			FlushPackageFile(NonSandboxPath, &PackageName);
 
-			if (LowerLevel->FileExists(*It->GetNonSandboxPath()))
+			if (LowerLevel->FileExists(*NonSandboxPath))
 			{
 				if (!PackageName.IsNone())
 				{
 					OutPackagesPendingHotReload.Add(PackageName);
 				}
-				FileChanges.Emplace(It->GetNonSandboxPath(), FFileChangeData::FCA_Added);
+				FileChanges.Emplace(NonSandboxPath, FFileChangeData::FCA_Added);
 			}
 			It.RemoveCurrent();
 			continue;
@@ -961,31 +973,31 @@ void FSandboxPlatformFile::DiscardSandbox(
 	}
 }
 
-TArray<FString> FSandboxPlatformFile::GatherSandboxChangedFilenames(TFunctionRef<bool(const FSandboxedPlatformFilePath& InFile)> InShouldSkipFile) const
+TArray<FString> FSandboxPlatformFile::GatherSandboxChangedFilenames() const
 {
 	TArray<FString> ChangedFiles;
 	
 	// Gather deleted paths
 	{
-		FScopeLock DeletedSandboxPathsLock(&DeletedSandboxPathsCS);
+		const FScopeLock DeletedSandboxPathsLock(&DeletedSandboxPathsCS);
 		ChangedFiles.Reserve(DeletedSandboxPaths.Num());
 		Algo::TransformIf(DeletedSandboxPaths, ChangedFiles,
-			[this, &InShouldSkipFile](const FSandboxedPlatformFilePath& DeletedPath)
+			[this](const FNonSandboxPath& DeletedPath)
 			{
-				return !InShouldSkipFile(DeletedPath) && LowerLevel->FileExists(*DeletedPath.GetNonSandboxPath());
+				return LowerLevel->FileExists(*DeletedPath);
 			},
-			[this](const FSandboxedPlatformFilePath& DeletedPath)
+			[this](const FNonSandboxPath& DeletedPath)
 			{
-				return DeletedPath.GetNonSandboxPath();
+				return DeletedPath;
 			});
 	}
 
 	// Gather mounts files
 	{
-		FScopeLock SandboxMountPointsLock(&SandboxMountPointsCS);
+		const FScopeLock SandboxMountPointsLock(&SandboxMountPointsCS);
 		for (const FSandboxMountPoint& MountPoint : SandboxMountPoints)
 		{
-			LowerLevel->IterateDirectoryRecursively(*MountPoint.Path.GetSandboxPath(), [this, &InShouldSkipFile, &MountPoint, &ChangedFiles](const TCHAR* InFilenameOrDirectory, const bool InIsDirectory) -> bool
+			LowerLevel->IterateDirectoryRecursively(*MountPoint.Path.GetSandboxPath(), [this, &MountPoint, &ChangedFiles](const TCHAR* InFilenameOrDirectory, const bool InIsDirectory) -> bool
 			{
 				if (InIsDirectory)
 				{
@@ -995,11 +1007,7 @@ TArray<FString> FSandboxPlatformFile::GatherSandboxChangedFilenames(TFunctionRef
 				const FSandboxedPlatformFilePath Path = FSandboxedPlatformFilePath::CreateNonSandboxPath(
 					FPaths::ConvertRelativePathToFull(InFilenameOrDirectory), MountPoint.Path
 					);
-				if (!InShouldSkipFile(Path))
-				{
-					ChangedFiles.Add(Path.GetNonSandboxPath());
-				}
-					
+				ChangedFiles.Add(Path.GetNonSandboxPath());
 				return true; // continue iteration
 			});
 		}
@@ -1041,11 +1049,17 @@ void FSandboxPlatformFile::RegisterContentMountPath(const FString& InAssetPath,
 	{
 		const FScopeLock SandboxMountPointsLock(&SandboxMountPointsCS);
 
-		FSandboxMountPoint& SandboxMountPoint = SandboxMountPoints.Add_GetRef(FSandboxMountPoint{ FSandboxedPlatformFilePath(MoveTemp(AbsoluteNonSandboxPath), MoveTemp(AbsoluteSandboxPath)) });
+		FSandboxMountPoint& SandboxMountPoint = SandboxMountPoints.Add_GetRef(
+			FSandboxMountPoint{ FSandboxedPlatformFilePath(MoveTemp(AbsoluteNonSandboxPath), MoveTemp(AbsoluteSandboxPath)) }
+			);
 #if WITH_EDITOR
 		if (IDirectoryWatcher* DirectoryWatcher = GetDirectoryWatcher())
 		{
-			DirectoryWatcher->RegisterDirectoryChangedCallback_Handle(SandboxMountPoint.Path.GetSandboxPath(), IDirectoryWatcher::FDirectoryChanged::CreateRaw(this, &FSandboxPlatformFile::OnDirectoryChanged, SandboxMountPoint.Path), SandboxMountPoint.OnDirectoryChangedHandle, IDirectoryWatcher::IncludeDirectoryChanges);
+			DirectoryWatcher->RegisterDirectoryChangedCallback_Handle(
+				SandboxMountPoint.Path.GetSandboxPath(), 
+				IDirectoryWatcher::FDirectoryChanged::CreateRaw(this, &FSandboxPlatformFile::OnSandboxDirectoryChanged, SandboxMountPoint.Path), 
+				SandboxMountPoint.OnDirectoryChangedHandle, IDirectoryWatcher::IncludeDirectoryChanges
+				);
 		}
 #endif
 	}
@@ -1054,6 +1068,19 @@ void FSandboxPlatformFile::RegisterContentMountPath(const FString& InAssetPath,
 void FSandboxPlatformFile::UnregisterContentMountPath(const FString& InAssetPath, const FString& InFilesystemPath)
 {
 	const FString AbsoluteSandboxPath = FPaths::ConvertRelativePathToFull(GetSandboxRootPath() / InAssetPath) / TEXT("");
+	{
+		const FScopeLock DeletedSandboxPathsLock(&DeletedSandboxPathsCS);
+		for (auto It = DeletedSandboxPaths.CreateIterator(); It; ++It)
+		{
+			const FSandboxedPlatformFilePath SandboxPath = ToSandboxPath_Absolute(*It);
+			if (SandboxPath.GetSandboxPath().StartsWith(AbsoluteSandboxPath))
+			{
+				It.RemoveCurrent();
+				continue;
+			}
+		}
+	}
+	
 	{
 		const FScopeLock SandboxMountPointsLock(&SandboxMountPointsCS);
 		SandboxMountPoints.RemoveAll([&AbsoluteSandboxPath](FSandboxMountPoint& InSandboxMountPoint) -> bool
@@ -1072,17 +1099,6 @@ void FSandboxPlatformFile::UnregisterContentMountPath(const FString& InAssetPath
 			return bShouldRemove;
 		});
 	}
-	{
-		const FScopeLock DeletedSandboxPathsLock(&DeletedSandboxPathsCS);
-		for (auto It = DeletedSandboxPaths.CreateIterator(); It; ++It)
-		{
-			if (It->GetSandboxPath().StartsWith(AbsoluteSandboxPath))
-			{
-				It.RemoveCurrent();
-				continue;
-			}
-		}
-	}
 	LowerLevel->DeleteDirectoryRecursively(*AbsoluteSandboxPath);
 }
 
@@ -1139,22 +1155,38 @@ FSandboxedPlatformFilePath FSandboxPlatformFile::FromSandboxPath_Absolute(FStrin
 bool FSandboxPlatformFile::IsPathDeleted(const FSandboxedPlatformFilePath& InPath) const
 {
 	FScopeLock DeletedSandboxPathsLock(&DeletedSandboxPathsCS);
-	return DeletedSandboxPaths.Contains(InPath);
+	return DeletedSandboxPaths.Contains(InPath.GetNonSandboxPath());
 }
 
-void FSandboxPlatformFile::SetPathDeleted(const FSandboxedPlatformFilePath& InPath, const bool bIsDeleted)
+void FSandboxPlatformFile::MarkPathDeleted(const FSandboxedPlatformFilePath& InPath)
 {
-	FScopeLock DeletedSandboxPathsLock(&DeletedSandboxPathsCS);
-	if (bIsDeleted)
+	const FNonSandboxPath& NonSandboxPath = InPath.GetNonSandboxPath();
+	const bool bIsNonSandbox = LowerLevel->FileExists(*NonSandboxPath) || LowerLevel->DirectoryExists(*NonSandboxPath);
+	// If the file or directory does not exist, it cannot be marked as deleted. This ensure is a sanity check that our logic is correct.
+	if (!ensure(bIsNonSandbox))
 	{
-		DeletedSandboxPaths.Add(InPath);
+		return;
 	}
-	else
+	
+	bool bAlreadyInSet = false;
+	DeletedSandboxPaths.Add(InPath.GetNonSandboxPath(), &bAlreadyInSet);
+	const bool bMadeChanges = !bAlreadyInSet;
+	
+	if (bMadeChanges)
 	{
-		DeletedSandboxPaths.Remove(InPath);
+		OnNonSandboxFileDeletedDelegate.Broadcast(InPath, true);
 	}
+}
+
+void FSandboxPlatformFile::ClearPathDeleted(const FSandboxedPlatformFilePath& InPath)
+{
+	const int32 NumRemoved = DeletedSandboxPaths.Remove(InPath.GetNonSandboxPath());
+	const bool bMadeChanges = NumRemoved > 0;
 	
-	OnFileDeletedDelegate.Broadcast(InPath, bIsDeleted);
+	if (bMadeChanges)
+	{
+		OnNonSandboxFileDeletedDelegate.Broadcast(InPath, false);
+	}
 }
 
 void FSandboxPlatformFile::NotifyFileDeleted(const FSandboxedPlatformFilePath& InPath)
@@ -1275,7 +1307,7 @@ TArray<FSandboxPlatformFile::FDirectoryItem> FSandboxPlatformFile::GetDirectoryC
 
 #if WITH_EDITOR
 
-void FSandboxPlatformFile::OnDirectoryChanged(const TArray<FFileChangeData>& FileChanges, FSandboxedPlatformFilePath MountPath)
+void FSandboxPlatformFile::OnSandboxDirectoryChanged(const TArray<FFileChangeData>& FileChanges, FSandboxedPlatformFilePath MountPath)
 {
 	if (!IsSandboxEnabled())
 	{
@@ -1290,17 +1322,15 @@ void FSandboxPlatformFile::OnDirectoryChanged(const TArray<FFileChangeData>& Fil
 		// Map the sandbox paths back to their original roots and notify the directory watcher
 		for (const FFileChangeData& FileChange : FileChanges)
 		{
-			const FSandboxedPlatformFilePath RemappedFilePath = FSandboxedPlatformFilePath::CreateNonSandboxPath(FPaths::ConvertRelativePathToFull(FileChange.Filename), MountPath);
+			const FSandboxedPlatformFilePath RemappedFilePath = FSandboxedPlatformFilePath::CreateNonSandboxPath(
+				FPaths::ConvertRelativePathToFull(FileChange.Filename), MountPath
+				);
 			RemappedFileChanges.Emplace(RemappedFilePath.GetNonSandboxPath(), FileChange.Action);
 
 			// Make sure the deleted state of this item is synchronized correctly
 			if (FileChange.Action == FFileChangeData::FCA_Added)
 			{
-				SetPathDeleted(RemappedFilePath, false);
-			}
-			else if (FileChange.Action == FFileChangeData::FCA_Removed)
-			{
-				SetPathDeleted(RemappedFilePath, true);
+				ClearPathDeleted(RemappedFilePath);
 			}
 		}
 
diff --git a/Engine/Plugins/Developer/Sandbox/FileSandbox/Source/FileSandboxCore/Private/Sandbox/Platform/SandboxPlatformFile.h b/Engine/Plugins/Developer/Sandbox/FileSandbox/Source/FileSandboxCore/Private/Sandbox/Platform/SandboxPlatformFile.h
--- a/Engine/Plugins/Developer/Sandbox/FileSandbox/Source/FileSandboxCore/Private/Sandbox/Platform/SandboxPlatformFile.h
+++ b/Engine/Plugins/Developer/Sandbox/FileSandbox/Source/FileSandboxCore/Private/Sandbox/Platform/SandboxPlatformFile.h
@@ -13,9 +13,9 @@ namespace UE::FileSandboxCore
 {
 struct FPersistResult;
 struct FPersistArgs;
-
+	
 DECLARE_MULTICAST_DELEGATE_TwoParams(FDeletionStateChanged, const FSandboxedPlatformFilePath&, bool /* bIsDeleted */);
-
+	
 class FSandboxPlatformFile : public IPlatformFile
 {
 public:
@@ -84,13 +84,16 @@ public:
 	 * Gather all files changes that are currently in the sandbox.
 	 * @return Array of filepaths in no particular order.
 	 */
-	TArray<FString> GatherSandboxChangedFilenames(TFunctionRef<bool(const FSandboxedPlatformFilePath& InFile)> InShouldSkipFile = [](auto){ return true; }) const;
+	TArray<FString> GatherSandboxChangedFilenames() const;
 	
 	/** @return Whether the given package file was deleted and exists on non sandbox path. */
 	bool DeletedPackageExistsInNonSandbox(FString InFilename) const;
 
-	/** Invoked when a file is deleted. */
-	FDeletionStateChanged& OnDeletionStateChanged() { return OnFileDeletedDelegate; }
+	/**
+	 * Invoked when a file from a non-sandbox path is deleted, i.e. a file that existed before the sandbox is deleted.
+	 * Specifically: NOT invoked for files that were added during sandbox, and then deleted again.
+	 */
+	FDeletionStateChanged& OnNonSandboxFileDeleted() { return OnNonSandboxFileDeletedDelegate; }
 	
 private:
 	
@@ -128,10 +131,12 @@ private:
 	/** Resolve the given path to its non-sandbox path (if any) from an absolute filename */
 	FSandboxedPlatformFilePath FromSandboxPath_Absolute(FString InFilename) const;
 
-	/** Check whether the given absolute sandbox path has been explicitly deleted from the sandbox */
+	/** @return Whether a file or directory that existed in the non-sandbox has been explicitly deleted by the sandbox. */
 	bool IsPathDeleted(const FSandboxedPlatformFilePath& InPath) const;
-	/** Set whether the given absolute sandbox path has been explicitly deleted from the sandbox */
-	void SetPathDeleted(const FSandboxedPlatformFilePath& InPath, const bool bIsDeleted);
+	/** Set that the non-sandbox file or directory is explicitly deleted by the sandbox. */
+	void MarkPathDeleted(const FSandboxedPlatformFilePath& InPath);
+	/** Clears the deletion state. The file or directory is no longer explicitly deleted by the sandbox. */
+	void ClearPathDeleted(const FSandboxedPlatformFilePath& InPath);
 
 	/** Notify that a file has been explicitly deleted from the sandbox */
 	void NotifyFileDeleted(const FSandboxedPlatformFilePath& InPath);
@@ -144,7 +149,7 @@ private:
 
 #if WITH_EDITOR
 	/** Called when a file in a sandbox directory changes on disk */
-	void OnDirectoryChanged(const TArray<FFileChangeData>& FileChanges, FSandboxedPlatformFilePath MountPath);
+	void OnSandboxDirectoryChanged(const TArray<FFileChangeData>& FileChanges, FSandboxedPlatformFilePath MountPath);
 #endif
 
 	/** Copy file using source control (if enabled) */
@@ -168,13 +173,21 @@ private:
 	/** Critical section protecting concurrent access to SandboxMountPoints */
 	mutable FCriticalSection SandboxMountPointsCS;
 
-	/** Set of absolute sandbox paths that have been explicitly deleted from the sandbox and shouldn't fallback to the non-sandbox items */
-	TSet<FSandboxedPlatformFilePath> DeletedSandboxPaths;
+	/**
+	 * Set of absolute paths (file or directory) that have been explicitly deleted by the sandbox, 
+	 * i.e. when persisting these files should be deleted from the non-sandbox.
+	 * 
+	 * This contains the absolute path to the non-sandbox file or directory.
+	 */
+	TSet<FSandboxPath> DeletedSandboxPaths;
 	/** Critical section protecting concurrent access to DeletedSandboxPaths */
 	mutable FCriticalSection DeletedSandboxPathsCS;
 	
-	/** Invoked when a file is deleted. */
-	FDeletionStateChanged OnFileDeletedDelegate;
+	/**
+	 * Invoked when a file from a non-sandbox path is deleted, i.e. a file that existed before the sandbox is deleted.
+	 * Specifically: NOT invoked for files that were added during sandbox, and then deleted again.
+	 */
+	FDeletionStateChanged OnNonSandboxFileDeletedDelegate;
 };
 }
 
diff --git a/Engine/Plugins/Developer/Sandbox/FileSandbox/Source/FileSandboxCore/Private/Sandbox/Platform/SandboxedPlatformFilePath.h b/Engine/Plugins/Developer/Sandbox/FileSandbox/Source/FileSandboxCore/Private/Sandbox/Platform/SandboxedPlatformFilePath.h
--- a/Engine/Plugins/Developer/Sandbox/FileSandbox/Source/FileSandboxCore/Private/Sandbox/Platform/SandboxedPlatformFilePath.h
+++ b/Engine/Plugins/Developer/Sandbox/FileSandbox/Source/FileSandboxCore/Private/Sandbox/Platform/SandboxedPlatformFilePath.h
@@ -6,6 +6,9 @@
 
 namespace UE::FileSandboxCore
 {
+using FNonSandboxPath = FString;
+using FSandboxPath = FString;
+	
 /** A pair of file paths: a path to the original file and a path to the sandboxed version, if it exists. */
 class FSandboxedPlatformFilePath
 {
@@ -77,13 +80,13 @@ public:
 	}
 
 	/** Get the absolute sandbox path */
-	const FString& GetSandboxPath() const
+	const FSandboxPath& GetSandboxPath() const
 	{
 		return SandboxPath;
 	}
 
 	/** Get the absolute non-sandbox path */
-	const FString& GetNonSandboxPath() const
+	const FNonSandboxPath& GetNonSandboxPath() const
 	{
 		return NonSandboxPath;
 	}
@@ -96,9 +99,9 @@ public:
 private:
 	
 	/** Absolute sandbox path */
-	FString SandboxPath;
+	FSandboxPath SandboxPath;
 
 	/** Absolute non-sandbox path */
-	FString NonSandboxPath;
+	FNonSandboxPath NonSandboxPath;
 };
 }
diff --git a/Engine/Plugins/Developer/Sandbox/FileSandbox/Source/FileSandboxCore/Private/Sandbox/SandboxInstance.cpp b/Engine/Plugins/Developer/Sandbox/FileSandbox/Source/FileSandboxCore/Private/Sandbox/SandboxInstance.cpp
--- a/Engine/Plugins/Developer/Sandbox/FileSandbox/Source/FileSandboxCore/Private/Sandbox/SandboxInstance.cpp
+++ b/Engine/Plugins/Developer/Sandbox/FileSandbox/Source/FileSandboxCore/Private/Sandbox/SandboxInstance.cpp
@@ -61,7 +61,7 @@ FSandboxInstance::FSandboxInstance(
 {
 	PlatformFile->Initialize(&InitArgs.GetPlatformFile(), TEXT(""));
 	ReplayChanges(ManifestContent.DeletedFiles.GetFiles(), ManifestContent.ModifiedFiles.GetFiles(), ManifestContent.AddedFiles.GetFiles());
-	PlatformFile->OnDeletionStateChanged().AddRaw(this, &FSandboxInstance::OnFileDeletionStateChanged);
+	PlatformFile->OnNonSandboxFileDeleted().AddRaw(this, &FSandboxInstance::OnFileDeletionStateChanged);
 }
 	
 FSandboxInstance::~FSandboxInstance()
@@ -129,15 +129,7 @@ FDiscardResult FSandboxInstance::DiscardSandbox(const FDiscardArgs& InArgs)
 
 TArray<FString> FSandboxInstance::GatherChangedFiles() const
 {
-	return PlatformFile->GatherSandboxChangedFilenames([this](const FSandboxedPlatformFilePath& InFilePath)
-	{
-		// TODO DP: Do we need this?
-		/*const FString& NonSandboxPath = InFilePath.GetNonSandboxPath();
-		const FDateTime* ModifiedTime = ManifestContent.PersistedFiles.Find(NonSandboxPath);
-		const bool bHasBeenPersisted = ModifiedTime && PlatformFile->GetTimeStamp(*NonSandboxPath) <= *ModifiedTime;
-		return !bHasBeenPersisted;*/
-		return false;
-	});
+	return PlatformFile->GatherSandboxChangedFilenames();
 }
 
 bool FSandboxInstance::DeletedPackageExistsInNonSandbox(const FString& InFilename) const
