File Sandbox: Misc cleanup
- Add misc functions to input and output structs
- Added ArgsFwds.h
- Implemented missing functions on ISandboxManager.h

#Synced-CL 47637010

diff --git a/Engine/Plugins/Developer/Sandbox/FileSandbox/Source/FileSandboxCore/Private/ISandboxManager.cpp b/Engine/Plugins/Developer/Sandbox/FileSandbox/Source/FileSandboxCore/Private/ISandboxManager.cpp
new file mode 100644
--- /dev/null
+++ b/Engine/Plugins/Developer/Sandbox/FileSandbox/Source/FileSandboxCore/Private/ISandboxManager.cpp
@@ -0,0 +1,39 @@
+﻿// Copyright Epic Games, Inc. All Rights Reserved.
+
+#include "ISandboxManager.h"
+
+#include "Types/Manager/DeleteSandboxByDirectoryArgs.h"
+#include "Types/Manager/DeleteSandboxByNameArgs.h"
+#include "Types/Manager/DeleteSandboxResult.h"
+#include "Types/Manager/LoadSandboxByDirectoryArgs.h"
+#include "Types/Manager/LoadSandboxByNameArgs.h"
+#include "Types/Manager/LoadSandboxError.h"
+#include "Types/Manager/SandboxCreationResult.h"
+#include "Utils/SandboxDirectoryUtils.h"
+
+namespace UE::FileSandboxCore
+{
+FLoadSandboxResult ISandboxManager::LoadNamedSandbox(const FLoadSandboxByNameArgs& InArgs)
+{
+	TOptional<FString> FullPath = FindSandboxByName(InArgs.Name, InArgs.BaseDirectory);
+	if (!FullPath)
+	{
+		return { MakeError<FLoadSandboxError>(FLoadSandboxError{ ELoadSandboxLoadErrorReason::InvalidName }) };
+	}
+	
+	FLoadSandboxByDirectoryArgs LoadByDirArgs(MoveTemp(*FullPath));
+	LoadByDirArgs.InitArgs = InArgs.InitArgs;
+	return LoadSandbox(LoadByDirArgs);
+}
+
+FDeleteSandboxResult ISandboxManager::DeleteNamedSandbox(const FDeleteSandboxByNameArgs& InArgs)
+{
+	TOptional<FString> FullPath = FindSandboxByName(InArgs.Name, InArgs.BaseDirectory);
+	if (!FullPath)
+	{
+		return FDeleteSandboxResult(EDeleteSandboxErrorCode::InvalidName);
+	}
+	
+	return DeleteSandbox(MoveTemp(*FullPath));
+}
+}
diff --git a/Engine/Plugins/Developer/Sandbox/FileSandbox/Source/FileSandboxCore/Private/Sandbox/SandboxManager.cpp b/Engine/Plugins/Developer/Sandbox/FileSandbox/Source/FileSandboxCore/Private/Sandbox/SandboxManager.cpp
--- a/Engine/Plugins/Developer/Sandbox/FileSandbox/Source/FileSandboxCore/Private/Sandbox/SandboxManager.cpp
+++ b/Engine/Plugins/Developer/Sandbox/FileSandbox/Source/FileSandboxCore/Private/Sandbox/SandboxManager.cpp
@@ -2,14 +2,18 @@
 
 #include "SandboxManager.h"
 
-#include "LogFileSandbox.h"
 #include "GenericPlatform/GenericPlatformFile.h"
 #include "HAL/PlatformFileManager.h"
+#include "LogFileSandbox.h"
 #include "Misc/CoreDelegates.h"
 #include "Misc/Paths.h"
+#include "Types/Manager/DeleteSandboxByDirectoryArgs.h"
 #include "Types/Manager/DeleteSandboxResult.h"
 #include "Types/Manager/LoadSandboxByDirectoryArgs.h"
+#include "Types/Manager/LoadSandboxError.h"
 #include "Types/Manager/NewSandboxArgs.h"
+#include "Types/Manager/NewSandboxError.h"
+#include "Types/Manager/SandboxCreationResult.h"
 #include "Utils/FileSandboxUtils.h"
 
 namespace UE::FileSandboxCore
@@ -78,20 +82,20 @@ FLoadSandboxResult FSandboxManager::LoadSandbox(const FLoadSandboxByDirectoryArg
 	return { MakeValue(Instance) };
 }
 
-FDeleteSandboxResult FSandboxManager::DeleteSandbox(const FString& InSandboxDirectory)
+FDeleteSandboxResult FSandboxManager::DeleteSandbox(const FDeleteSandboxByDirectoryArgs& InArgs)
 {
-	if (!IsRootSandboxDirectory(InSandboxDirectory))
+	if (!IsRootSandboxDirectory(InArgs.Directory))
 	{
 		return EDeleteSandboxErrorCode::InvalidDirectory;
 	}
 	
-	if (ActiveSandbox && FPaths::IsSamePath(ActiveSandbox->GetRootDirectory(), InSandboxDirectory))
+	if (ActiveSandbox && FPaths::IsSamePath(ActiveSandbox->GetRootDirectory(), InArgs.Directory))
 	{
 		LeaveSandbox();
 	}
 
 	IPlatformFile& PlatformFile = FPlatformFileManager::Get().GetPlatformFile();
-	return PlatformFile.DeleteDirectoryRecursively(*InSandboxDirectory)
+	return PlatformFile.DeleteDirectoryRecursively(*InArgs.Directory)
 		? EDeleteSandboxErrorCode::Success
 		: EDeleteSandboxErrorCode::IOError;
 }
diff --git a/Engine/Plugins/Developer/Sandbox/FileSandbox/Source/FileSandboxCore/Private/Sandbox/SandboxManager.h b/Engine/Plugins/Developer/Sandbox/FileSandbox/Source/FileSandboxCore/Private/Sandbox/SandboxManager.h
--- a/Engine/Plugins/Developer/Sandbox/FileSandbox/Source/FileSandboxCore/Private/Sandbox/SandboxManager.h
+++ b/Engine/Plugins/Developer/Sandbox/FileSandbox/Source/FileSandboxCore/Private/Sandbox/SandboxManager.h
@@ -19,7 +19,7 @@ public:
 	//~ Begin ISandboxManager Interface
 	virtual FNewSandboxResult CreateNewSandbox(const FNewSandboxArgs& InArgs) override;
 	virtual FLoadSandboxResult LoadSandbox(const FLoadSandboxByDirectoryArgs& InLoadArgs) override;
-	virtual FDeleteSandboxResult DeleteSandbox(const FString& InSandboxDirectory) override;
+	virtual FDeleteSandboxResult DeleteSandbox(const FDeleteSandboxByDirectoryArgs& InArgs) override;
 	virtual void LeaveSandbox(const FLeaveSandboxArgs& InLeaveArgs = FLeaveSandboxArgs()) override;
 	virtual ISandboxInstance* GetActiveSandboxInstance() const override { return ActiveSandbox.Get(); }
 	virtual FSandboxInstanceEvent& OnPostSandboxStartup() override { return OnPostSandboxStartupDelegate; }
diff --git a/Engine/Plugins/Developer/Sandbox/FileSandbox/Source/FileSandboxCore/Private/Types/Manager/LoadSandboxError.cpp b/Engine/Plugins/Developer/Sandbox/FileSandbox/Source/FileSandboxCore/Private/Types/Manager/LoadSandboxError.cpp
new file mode 100644
--- /dev/null
+++ b/Engine/Plugins/Developer/Sandbox/FileSandbox/Source/FileSandboxCore/Private/Types/Manager/LoadSandboxError.cpp
@@ -0,0 +1,21 @@
+﻿// Copyright Epic Games, Inc. All Rights Reserved.
+
+#include "Types/Manager/LoadSandboxError.h"
+
+#include "Containers/UnrealString.h"
+
+namespace UE::FileSandboxCore
+{
+FString LexToString(ELoadSandboxLoadErrorReason InReason)
+{
+	static_assert(static_cast<int32>(ELoadSandboxLoadErrorReason::Count) == 4, "Update this switch");
+	switch (InReason)
+	{
+	case ELoadSandboxLoadErrorReason::InvalidDirectory: return TEXT("InvalidDirectory");
+	case ELoadSandboxLoadErrorReason::InvalidName: return TEXT("InvalidName");
+	case ELoadSandboxLoadErrorReason::IOError: return TEXT("IOError");
+	case ELoadSandboxLoadErrorReason::Unspecified: return TEXT("Unspecified");
+	default: checkNoEntry(); return TEXT("Unknown");
+	}
+}
+}
diff --git a/Engine/Plugins/Developer/Sandbox/FileSandbox/Source/FileSandboxCore/Private/Types/Manager/NewSandboxError.cpp b/Engine/Plugins/Developer/Sandbox/FileSandbox/Source/FileSandboxCore/Private/Types/Manager/NewSandboxError.cpp
new file mode 100644
--- /dev/null
+++ b/Engine/Plugins/Developer/Sandbox/FileSandbox/Source/FileSandboxCore/Private/Types/Manager/NewSandboxError.cpp
@@ -0,0 +1,21 @@
+﻿// Copyright Epic Games, Inc. All Rights Reserved.
+
+#include "Types/Manager/NewSandboxError.h"
+
+#include "Containers/UnrealString.h"
+#include "Misc/AssertionMacros.h"
+
+namespace UE::FileSandboxCore
+{
+FString LexToString(ENewSandboxErrorReason InReason)
+{
+	static_assert(static_cast<int32>(ENewSandboxErrorReason::Count) == 3, "Update this switch");
+	switch (InReason)
+	{
+	case ENewSandboxErrorReason::UnsuitablePath: return TEXT("UnsuitablePath");
+	case ENewSandboxErrorReason::IOError: return TEXT("IOError");
+	case ENewSandboxErrorReason::Unspecified: return TEXT("Unspecified");
+	default: checkNoEntry(); return TEXT("Unknown");
+	}
+}
+}
diff --git a/Engine/Plugins/Developer/Sandbox/FileSandbox/Source/FileSandboxCore/Public/ISandboxManager.h b/Engine/Plugins/Developer/Sandbox/FileSandbox/Source/FileSandboxCore/Public/ISandboxManager.h
--- a/Engine/Plugins/Developer/Sandbox/FileSandbox/Source/FileSandboxCore/Public/ISandboxManager.h
+++ b/Engine/Plugins/Developer/Sandbox/FileSandbox/Source/FileSandboxCore/Public/ISandboxManager.h
@@ -3,20 +3,15 @@
 #pragma once
 
 #include "Delegates/Delegate.h"
+#include "Types/ArgsFwds.h"
 #include "Types/Manager/LeaveSandboxArgs.h"
 #include "Types/Manager/SandboxCreationResult.h"
 
+#define UE_API FILESANDBOXCORE_API
+
 namespace UE::FileSandboxCore
 {
-class FSandboxHandle;
 class ISandboxInstance;
-struct FDeleteSandboxResult;
-struct FLoadSandboxByDirectoryArgs;
-struct FLoadSandboxByNameArgs;
-struct FLoadSandboxError;
-struct FNewSandboxArgs;
-struct FNewSandboxError;
-struct FSandboxCreationResult;
 	
 DECLARE_MULTICAST_DELEGATE_OneParam(FSandboxInstanceEvent, ISandboxInstance&);
 using FSandboxShutdownEvent = FSimpleMulticastDelegate;
@@ -81,17 +76,17 @@ public:
 	 * @note Do not keep any raw reference to ISandboxInstance as it turns stale when the sandbox is shutdown.
 	 * @see FWeakSandboxPtr, GetGlobalSandboxInstance.
 	 */
-	FLoadSandboxResult LoadNamedSandbox(const FLoadSandboxByNameArgs& InName);
+	UE_API FLoadSandboxResult LoadNamedSandbox(const FLoadSandboxByNameArgs& InArgs);
 	
 	/** Advanced version of LoadNamedSandbox for when you know the specific root directory of a sandbox. */
 	virtual FLoadSandboxResult LoadSandbox(const FLoadSandboxByDirectoryArgs& InLoadArgs) = 0;
 	
 	
 	/** Deletes the sandbox with the given name, if found. It is valid to delete a sandbox if it is currently valid. */
-	FDeleteSandboxResult DeleteNamedSandbox(const FString& InSandboxName);
+	UE_API FDeleteSandboxResult DeleteNamedSandbox(const FDeleteSandboxByNameArgs& InArgs);
 	
 	/** Advanced version of DeleteNamedSandbox for when you know the specific root directory of a sandbox. */
-	virtual FDeleteSandboxResult DeleteSandbox(const FString& InSandboxDirectory) = 0;
+	virtual FDeleteSandboxResult DeleteSandbox(const FDeleteSandboxByDirectoryArgs& InArgs) = 0;
 	
 	
 	/** Leaves the current sandbox. */
@@ -123,3 +118,5 @@ public:
 	virtual ~ISandboxManager() = default;
 };
 }
+
+#undef UE_API
\ No newline at end of file
diff --git a/Engine/Plugins/Developer/Sandbox/FileSandbox/Source/FileSandboxCore/Public/Types/ArgsFwds.h b/Engine/Plugins/Developer/Sandbox/FileSandbox/Source/FileSandboxCore/Public/Types/ArgsFwds.h
new file mode 100644
--- /dev/null
+++ b/Engine/Plugins/Developer/Sandbox/FileSandbox/Source/FileSandboxCore/Public/Types/ArgsFwds.h
@@ -0,0 +1,23 @@
+﻿// Copyright Epic Games, Inc. All Rights Reserved.
+
+#pragma once
+
+namespace UE::FileSandboxCore
+{
+struct FDeleteSandboxByDirectoryArgs;
+struct FDeleteSandboxByNameArgs;
+struct FDeleteSandboxResult;
+struct FLeaveSandboxArgs;
+struct FLoadSandboxByDirectoryArgs;
+struct FLoadSandboxByNameArgs;
+struct FNewSandboxArgs;
+	
+struct FLoadSandboxError;
+struct FNewSandboxError;
+struct FSandboxCreationResult;
+	
+template<typename TErrorType>
+struct TSandboxCreationResult;
+using FNewSandboxResult = TSandboxCreationResult<FNewSandboxError>;
+using FLoadSandboxResult = TSandboxCreationResult<FLoadSandboxError>;
+}
\ No newline at end of file
diff --git a/Engine/Plugins/Developer/Sandbox/FileSandbox/Source/FileSandboxCore/Public/Types/Manager/DeleteSandboxByDirectoryArgs.h b/Engine/Plugins/Developer/Sandbox/FileSandbox/Source/FileSandboxCore/Public/Types/Manager/DeleteSandboxByDirectoryArgs.h
new file mode 100644
--- /dev/null
+++ b/Engine/Plugins/Developer/Sandbox/FileSandbox/Source/FileSandboxCore/Public/Types/Manager/DeleteSandboxByDirectoryArgs.h
@@ -0,0 +1,19 @@
+﻿// Copyright Epic Games, Inc. All Rights Reserved.
+
+#pragma once
+
+#include "Containers/UnrealString.h"
+
+namespace UE::FileSandboxCore
+{
+/** Args for deleting sandbox by directory. */
+struct FDeleteSandboxByDirectoryArgs
+{
+	/** The root directory of the sandbox (i.e. the one containing the manifest file). */
+	FString Directory;
+
+	FDeleteSandboxByDirectoryArgs(FString InDirectory) 
+		: Directory(MoveTemp(InDirectory))
+	{}
+};
+}
\ No newline at end of file
diff --git a/Engine/Plugins/Developer/Sandbox/FileSandbox/Source/FileSandboxCore/Public/Types/Manager/DeleteSandboxByNameArgs.h b/Engine/Plugins/Developer/Sandbox/FileSandbox/Source/FileSandboxCore/Public/Types/Manager/DeleteSandboxByNameArgs.h
new file mode 100644
--- /dev/null
+++ b/Engine/Plugins/Developer/Sandbox/FileSandbox/Source/FileSandboxCore/Public/Types/Manager/DeleteSandboxByNameArgs.h
@@ -0,0 +1,26 @@
+﻿// Copyright Epic Games, Inc. All Rights Reserved.
+
+#pragma once
+
+#include "SandboxInitArgs.h"
+#include "Utils/SandboxDirectoryUtils.h"
+
+namespace UE::FileSandboxCore
+{
+/** Args for deleting sandbox by name. */
+struct FDeleteSandboxByNameArgs
+{
+	/** The name of the sandbox you want to load. */
+	FString Name;
+
+	/** The directory in which to search for the sandbox. */
+	FString BaseDirectory;
+
+	explicit FDeleteSandboxByNameArgs(FString InName, FString InBaseDirectory = GetBaseSandboxDirectory()) 
+		: Name(MoveTemp(InName))
+		, BaseDirectory(MoveTemp(InBaseDirectory))
+	{
+		check(!Name.IsEmpty());
+	}
+};
+}
diff --git a/Engine/Plugins/Developer/Sandbox/FileSandbox/Source/FileSandboxCore/Public/Types/Manager/LoadSandboxByNameArgs.h b/Engine/Plugins/Developer/Sandbox/FileSandbox/Source/FileSandboxCore/Public/Types/Manager/LoadSandboxByNameArgs.h
--- a/Engine/Plugins/Developer/Sandbox/FileSandbox/Source/FileSandboxCore/Public/Types/Manager/LoadSandboxByNameArgs.h
+++ b/Engine/Plugins/Developer/Sandbox/FileSandbox/Source/FileSandboxCore/Public/Types/Manager/LoadSandboxByNameArgs.h
@@ -14,12 +14,14 @@ struct FLoadSandboxByNameArgs
 	FString Name;
 
 	/** The directory in which to search for the sandbox. */
-	FString BaseDirectory = GetBaseSandboxDirectory();
+	FString BaseDirectory;
 	
 	/** Base arguments for instantiating the sandbox environment. */
-	FSandboxInitArgs BaseArgs;
+	FSandboxInitArgs InitArgs;
 
-	explicit FLoadSandboxByNameArgs(FString InName) : Name(MoveTemp(InName))
+	explicit FLoadSandboxByNameArgs(FString InName, FString InBaseDirectory = GetBaseSandboxDirectory()) 
+		: Name(MoveTemp(InName))
+		, BaseDirectory(MoveTemp(InBaseDirectory))
 	{
 		check(!Name.IsEmpty());
 	}
diff --git a/Engine/Plugins/Developer/Sandbox/FileSandbox/Source/FileSandboxCore/Public/Types/Manager/LoadSandboxError.h b/Engine/Plugins/Developer/Sandbox/FileSandbox/Source/FileSandboxCore/Public/Types/Manager/LoadSandboxError.h
--- a/Engine/Plugins/Developer/Sandbox/FileSandbox/Source/FileSandboxCore/Public/Types/Manager/LoadSandboxError.h
+++ b/Engine/Plugins/Developer/Sandbox/FileSandbox/Source/FileSandboxCore/Public/Types/Manager/LoadSandboxError.h
@@ -4,6 +4,8 @@
 
 #include "HAL/Platform.h"
 
+class FString;
+
 namespace UE::FileSandboxCore
 {
 enum class ELoadSandboxLoadErrorReason : uint8
@@ -17,7 +19,13 @@ enum class ELoadSandboxLoadErrorReason : uint8
 	IOError,
 
 	/** Dummy error code for initializing to a deterministic value. */
-	Unspecified
+	Unspecified,
+	
+	/** 
+	 * ADD NEW MEMBERS ABOVE. 
+	 * Not a real reason. 
+	 */
+	Count
 };
 
 /** Explains why a new sandbox could not be loaded. */ 
@@ -25,4 +33,7 @@ struct FLoadSandboxError
 {
 	ELoadSandboxLoadErrorReason Reason = ELoadSandboxLoadErrorReason::Unspecified;
 };
+	
+/** @return Error reason converted to a string. */
+FILESANDBOXCORE_API FString LexToString(ELoadSandboxLoadErrorReason InReason);
 }
\ No newline at end of file
diff --git a/Engine/Plugins/Developer/Sandbox/FileSandbox/Source/FileSandboxCore/Public/Types/Manager/NewSandboxError.h b/Engine/Plugins/Developer/Sandbox/FileSandbox/Source/FileSandboxCore/Public/Types/Manager/NewSandboxError.h
--- a/Engine/Plugins/Developer/Sandbox/FileSandbox/Source/FileSandboxCore/Public/Types/Manager/NewSandboxError.h
+++ b/Engine/Plugins/Developer/Sandbox/FileSandbox/Source/FileSandboxCore/Public/Types/Manager/NewSandboxError.h
@@ -4,6 +4,8 @@
 
 #include "HAL/Platform.h"
 
+class FString;
+
 namespace UE::FileSandboxCore
 {
 enum class ENewSandboxErrorReason : uint8
@@ -15,7 +17,13 @@ enum class ENewSandboxErrorReason : uint8
 	IOError,
 
 	/** Dummy error code for initializing to a deterministic value. */
-	Unspecified
+	Unspecified,
+
+	/** 
+	 * ADD NEW MEMBERS ABOVE. 
+	 * Not a real reason. 
+	 */
+	Count
 };
 
 /** Explains why a new sandbox could not be created. */ 
@@ -23,4 +31,7 @@ struct FNewSandboxError
 {
 	ENewSandboxErrorReason Reason = ENewSandboxErrorReason::Unspecified;
 };
-}
\ No newline at end of file
+
+/** @return Error reason converted to a string. */
+FILESANDBOXCORE_API FString LexToString(ENewSandboxErrorReason InReason);
+}
diff --git a/Engine/Plugins/Developer/Sandbox/FileSandbox/Source/FileSandboxCore/Public/Types/Manager/SandboxCreationResult.h b/Engine/Plugins/Developer/Sandbox/FileSandbox/Source/FileSandboxCore/Public/Types/Manager/SandboxCreationResult.h
--- a/Engine/Plugins/Developer/Sandbox/FileSandbox/Source/FileSandboxCore/Public/Types/Manager/SandboxCreationResult.h
+++ b/Engine/Plugins/Developer/Sandbox/FileSandbox/Source/FileSandboxCore/Public/Types/Manager/SandboxCreationResult.h
@@ -25,7 +25,4 @@ struct TSandboxCreationResult : TValueOrError<TNotNull<ISandboxInstance*>, TErro
 
 	operator ISandboxInstance*() const { return this->HasValue() ? this->GetValue() : nullptr; }
 };
-
-using FNewSandboxResult = TSandboxCreationResult<FNewSandboxError>;
-using FLoadSandboxResult = TSandboxCreationResult<FLoadSandboxError>;
 }
