File Sandbox: Introduce source control proxy.

When the engine is set up with source control, certain file changes do not go through IPlatformFile any more. For example, when deleting an asset, the a FDelete command is issued to source control instead. The intention is that e.g. p4 delete will delete the underlying file.

While in sandbox, source control operations should be silently absorbed. This is because we're not really touching the underlying files so it would be wrong if in source control files were marked for add, etc. Only once you decide to persist the sandbox, source control should be updated.

#Synced-CL 47637010

diff --git a/Engine/Plugins/Developer/Sandbox/FileSandbox/Source/FileSandboxCore/FileSandboxCore.Build.cs b/Engine/Plugins/Developer/Sandbox/FileSandbox/Source/FileSandboxCore/FileSandboxCore.Build.cs
--- a/Engine/Plugins/Developer/Sandbox/FileSandbox/Source/FileSandboxCore/FileSandboxCore.Build.cs
+++ b/Engine/Plugins/Developer/Sandbox/FileSandbox/Source/FileSandboxCore/FileSandboxCore.Build.cs
@@ -24,6 +24,14 @@ namespace UnrealBuildTool.Rules
 					"SourceControl", 
 				});
 
+			if (Target.bUsesSlate)
+			{
+				PrivateDependencyModuleNames.AddRange(new string[]
+				{
+					"SlateCore" // Needed for source control proxy returning SNullWidget::NullWidget
+				});
+			}
+
 			if (Target.bBuildEditor)
 			{
 				PrivateDependencyModuleNames.AddRange(
diff --git a/Engine/Plugins/Developer/Sandbox/FileSandbox/Source/FileSandboxCore/Private/Sandbox/Platform/SandboxPlatformFile.cpp b/Engine/Plugins/Developer/Sandbox/FileSandbox/Source/FileSandboxCore/Private/Sandbox/Platform/SandboxPlatformFile.cpp
--- a/Engine/Plugins/Developer/Sandbox/FileSandbox/Source/FileSandboxCore/Private/Sandbox/Platform/SandboxPlatformFile.cpp
+++ b/Engine/Plugins/Developer/Sandbox/FileSandbox/Source/FileSandboxCore/Private/Sandbox/Platform/SandboxPlatformFile.cpp
@@ -339,7 +339,10 @@ bool FSandboxPlatformFile::MoveFile(const TCHAR* To, const TCHAR* From)
 			bSuccess = LowerLevel->CopyFile(*ResolvedToPath.GetNonSandboxPath(), *ResolvedFromPath.GetNonSandboxPath());
 		}
 
-		if (bSuccess)
+		// ResolvedFromPath may be a file that was added in the sandbox, so the non-sandbox counterpart may not exist.
+		// We can get into this path when From = new .uasset, and To = a .temp file (some engine ops create temp, backup files of assets).
+		const bool bNeedsMarkForRemove = LowerLevel->FileExists(*ResolvedFromPath.GetNonSandboxPath());
+		if (bSuccess && bNeedsMarkForRemove)
 		{
 			MarkPathDeleted(ResolvedFromPath);
 			NotifyFileDeleted(ResolvedFromPath);
@@ -678,6 +681,7 @@ FString FSandboxPlatformFile::ConvertToAbsolutePathForExternalAppForWrite(const
 
 
 bool FSandboxPlatformFile::CopyFileWithSCC(
+	ISourceControlProvider& InSourceControlProvider, 
 	const TCHAR* InToFilename,
 	const TCHAR* InFromFilename,
 	const FPersistArgs& InParam,
@@ -687,12 +691,11 @@ bool FSandboxPlatformFile::CopyFileWithSCC(
 	// If this file maps to a package then we need to flush its linker so that we can overwrite the file on disk
 	FlushPackageFile(InToFilename);
 
-	ISourceControlProvider* SourceControlProvider = &ISourceControlModule::Get().GetProvider();
 	// Get the source control state of the destination file
 	FSourceControlStatePtr ToFileSCCState;
-	if (SourceControlProvider && SourceControlProvider->IsEnabled())
+	if (InSourceControlProvider.IsEnabled())
 	{
-		ToFileSCCState = SourceControlProvider->GetState(InToFilename, EStateCacheUsage::ForceUpdate);
+		ToFileSCCState = InSourceControlProvider.GetState(InToFilename, EStateCacheUsage::ForceUpdate);
 	}
 
 	// We don't need to do anything with source control if the file is already checked-out or added
@@ -701,14 +704,12 @@ bool FSandboxPlatformFile::CopyFileWithSCC(
 	// If the file can be checked-out, do so now
 	if (bRequiresSCCAction && ToFileSCCState->IsSourceControlled())
 	{
-		// the static analysis tool is not able to see that `SourceControlProvider` won't be null if `bRequiresSCCAction` is true
-		CA_ASSUME(SourceControlProvider != nullptr);
-		if (ToFileSCCState->CanCheckout() && SourceControlProvider->UsesCheckout())
+		if (ToFileSCCState->CanCheckout() && InSourceControlProvider.UsesCheckout())
 		{
 			TArray<FString> FilesToBeCheckedOut;
 			FilesToBeCheckedOut.Add(InToFilename);
 
-			if (SourceControlProvider->Execute(ISourceControlOperation::Create<FCheckOut>(), FilesToBeCheckedOut) != ECommandResult::Succeeded)
+			if (InSourceControlProvider.Execute(ISourceControlOperation::Create<FCheckOut>(), FilesToBeCheckedOut) != ECommandResult::Succeeded)
 			{
 				FText Failure = FText::Format(
 					NSLOCTEXT("Concert", "CheckoutFileFailure", "Failed to check-out file '{0}' from revision control when persiting sandbox state!"),
@@ -759,7 +760,7 @@ bool FSandboxPlatformFile::CopyFileWithSCC(
 		TArray<FString> FilesToBeAdded;
 		FilesToBeAdded.Add(InToFilename);
 
-		if (SourceControlProvider->Execute(ISourceControlOperation::Create<FMarkForAdd>(), FilesToBeAdded) != ECommandResult::Succeeded)
+		if (InSourceControlProvider.Execute(ISourceControlOperation::Create<FMarkForAdd>(), FilesToBeAdded) != ECommandResult::Succeeded)
 		{
 			FText Failure = FText::Format(
 				NSLOCTEXT("Concert", "AddFileFailure", "Failed to add file '{0}' to revision control when persiting sandbox state!"),
@@ -772,20 +773,14 @@ bool FSandboxPlatformFile::CopyFileWithSCC(
 	return true;
 }
 
-/** Delete file using source control */
-bool FSandboxPlatformFile::DeleteFileWithSCC(const TCHAR* InFilename, FPersistResult& OutResult)
+bool FSandboxPlatformFile::DeleteFileWithSCC(ISourceControlProvider& InSourceControlProvider, const TCHAR* InFilename, FPersistResult& OutResult)
 {
 	// If this file maps to a package then we need to flush its linker so that we can remove the file from disk
 	FlushPackageFile(InFilename);
 
-	ISourceControlProvider* SourceControlProvider = &ISourceControlModule::Get().GetProvider();
 	// Get the source control state of the file
-	FSourceControlStatePtr FileSCCState;
-	if (SourceControlProvider && SourceControlProvider->IsEnabled())
-	{
-		FileSCCState = SourceControlProvider->GetState(InFilename, EStateCacheUsage::ForceUpdate);
-	}
-
+	const FSourceControlStatePtr FileSCCState = InSourceControlProvider.IsEnabled() 
+		? InSourceControlProvider.GetState(InFilename, EStateCacheUsage::ForceUpdate) : nullptr;
 	// Try and let source control remove the file first
 	if (FileSCCState && FileSCCState->IsSourceControlled())
 	{
@@ -793,24 +788,22 @@ bool FSandboxPlatformFile::DeleteFileWithSCC(const TCHAR* InFilename, FPersistRe
 
 		if (bAdded || FileSCCState->IsCheckedOut())
 		{
-			if (SourceControlProvider->Execute(ISourceControlOperation::Create<FRevert>(), InFilename) != ECommandResult::Succeeded)
+			if (InSourceControlProvider.Execute(ISourceControlOperation::Create<FRevert>(), InFilename) != ECommandResult::Succeeded)
 			{
 				FText Failure = FText::Format(NSLOCTEXT("Concert", "RevertFileFailure", "Failed to revert file '{0}' to revision control when persisting sandbox state!"), FText::FromString(InFilename));
 				UE_LOG(LogFileSandbox, Warning, TEXT("%s"), *Failure.ToString());
 				OutResult.FailureReasons.Add(MoveTemp(Failure));
 				return false;
 			}
-
-			if (!bAdded)
-			{
-				if (SourceControlProvider->Execute(ISourceControlOperation::Create<FDelete>(), InFilename) != ECommandResult::Succeeded)
-				{
-					FText Failure = FText::Format(NSLOCTEXT("Concert", "DeleteSCCFileFailure", "Failed to delete file '{0}' from revision control when persiting sandbox state!"), FText::FromString(InFilename));
-					UE_LOG(LogFileSandbox, Warning, TEXT("%s"), *Failure.ToString());
-					OutResult.FailureReasons.Add(MoveTemp(Failure));
-					return false;
-				}
-			}
+		}
+		
+		// If the file was marked for add, then revert is guaranteed to remove the file; it is invalid to call delete on that file again.
+		if (!bAdded && InSourceControlProvider.Execute(ISourceControlOperation::Create<FDelete>(), InFilename) != ECommandResult::Succeeded)
+		{
+			FText Failure = FText::Format(NSLOCTEXT("Concert", "DeleteSCCFileFailure", "Failed to delete file '{0}' from revision control when persiting sandbox state!"), FText::FromString(InFilename));
+			UE_LOG(LogFileSandbox, Warning, TEXT("%s"), *Failure.ToString());
+			OutResult.FailureReasons.Add(MoveTemp(Failure));
+			return false;
 		}
 	}
 
@@ -825,7 +818,7 @@ bool FSandboxPlatformFile::DeleteFileWithSCC(const TCHAR* InFilename, FPersistRe
 	return true;
 }
 
-FPersistResult FSandboxPlatformFile::PersistSandbox(
+FPersistResult FSandboxPlatformFile::PersistSandbox(ISourceControlProvider& InSourceControlProvider,
 	FPersistArgs InParams, TFunctionRef<void(const FSandboxedPlatformFilePath&, EPersistedFileType)> InOnPersistFile
 	)
 {
@@ -838,8 +831,8 @@ FPersistResult FSandboxPlatformFile::PersistSandbox(
 		const FSandboxedPlatformFilePath FilePath = ToSandboxPath(File, true);
 		const EPersistedFileType Type = IsPathDeleted(FilePath) ? EPersistedFileType::Deleted : EPersistedFileType::Moved; 
 		const bool bSuccess = Type == EPersistedFileType::Deleted
-			? DeleteFileWithSCC(*FilePath.GetNonSandboxPath(), OutResult)
-			: CopyFileWithSCC(*FilePath.GetNonSandboxPath(), *FilePath.GetSandboxPath(), InParams, OutResult);
+			? DeleteFileWithSCC(InSourceControlProvider, *FilePath.GetNonSandboxPath(), OutResult)
+			: CopyFileWithSCC(InSourceControlProvider, *FilePath.GetNonSandboxPath(), *FilePath.GetSandboxPath(), InParams, OutResult);
 
 		// if the file operation was successful, mark the file as persisted
 		if (bSuccess)
@@ -1163,7 +1156,7 @@ void FSandboxPlatformFile::MarkPathDeleted(const FSandboxedPlatformFilePath& InP
 	const FNonSandboxPath& NonSandboxPath = InPath.GetNonSandboxPath();
 	const bool bIsNonSandbox = LowerLevel->FileExists(*NonSandboxPath) || LowerLevel->DirectoryExists(*NonSandboxPath);
 	// If the file or directory does not exist, it cannot be marked as deleted. This ensure is a sanity check that our logic is correct.
-	if (!ensure(bIsNonSandbox))
+	if (!ensureAlways(bIsNonSandbox))
 	{
 		return;
 	}
diff --git a/Engine/Plugins/Developer/Sandbox/FileSandbox/Source/FileSandboxCore/Private/Sandbox/Platform/SandboxPlatformFile.h b/Engine/Plugins/Developer/Sandbox/FileSandbox/Source/FileSandboxCore/Private/Sandbox/Platform/SandboxPlatformFile.h
--- a/Engine/Plugins/Developer/Sandbox/FileSandbox/Source/FileSandboxCore/Private/Sandbox/Platform/SandboxPlatformFile.h
+++ b/Engine/Plugins/Developer/Sandbox/FileSandbox/Source/FileSandboxCore/Private/Sandbox/Platform/SandboxPlatformFile.h
@@ -9,6 +9,8 @@
 #include "IDirectoryWatcher.h"
 #include "SandboxedPlatformFilePath.h"
 
+class ISourceControlProvider;
+
 namespace UE::FileSandboxCore
 {
 struct FPersistResult;
@@ -64,11 +66,12 @@ public:
 	/**
 	 *	Persist the file list from the sandbox state onto the real files Will mark files for which the operation was
 	 *	successful as persisted.
+	 *	@param InSourceControlProvider The source control provider to register the changes with
 	 *	@param InParams	Parameters controlling the persist options specified by the user.
 	 *	@param InOnPersistFile Invoked for each file successfully persisted.
 	 *	@return Returns a persist result object that contains result status.
 	 */
-	FPersistResult PersistSandbox(FPersistArgs InParams, TFunctionRef<void(const FSandboxedPlatformFilePath&, EPersistedFileType)> InOnPersistFile);
+	FPersistResult PersistSandbox(ISourceControlProvider& InSourceControlProvider, FPersistArgs InParams, TFunctionRef<void(const FSandboxedPlatformFilePath&, EPersistedFileType)> InOnPersistFile);
 
 	/** 
 	 *	Discard the sandbox state
@@ -153,9 +156,9 @@ private:
 #endif
 
 	/** Copy file using source control (if enabled) */
-	bool CopyFileWithSCC(const TCHAR* InToFilename, const TCHAR* InFromFilename, const FPersistArgs& InParam, FPersistResult& OutResult);
+	bool CopyFileWithSCC(ISourceControlProvider& InSourceControlProvider, const TCHAR* InToFilename, const TCHAR* InFromFilename, const FPersistArgs& InParam, FPersistResult& OutResult);
 	/** Delete file using source control (if enabled) */
-	bool DeleteFileWithSCC(const TCHAR* InFilename, FPersistResult& OutResult);
+	bool DeleteFileWithSCC(ISourceControlProvider& InSourceControlProvider, const TCHAR* InFilename, FPersistResult& OutResult);
 
 	FString GetSandboxRootPath() const;
 		
diff --git a/Engine/Plugins/Developer/Sandbox/FileSandbox/Source/FileSandboxCore/Private/Sandbox/SandboxInstance.cpp b/Engine/Plugins/Developer/Sandbox/FileSandbox/Source/FileSandboxCore/Private/Sandbox/SandboxInstance.cpp
--- a/Engine/Plugins/Developer/Sandbox/FileSandbox/Source/FileSandboxCore/Private/Sandbox/SandboxInstance.cpp
+++ b/Engine/Plugins/Developer/Sandbox/FileSandbox/Source/FileSandboxCore/Private/Sandbox/SandboxInstance.cpp
@@ -58,6 +58,7 @@ FSandboxInstance::FSandboxInstance(
 	, ManifestContent(MoveTemp(InManifestData))
 	, MetaData(MoveTemp(InMetaData))
 	, PlatformFile(MakeUnique<FSandboxPlatformFile>(RootDirectory))
+	, SourceControlSandbox(*this)
 {
 	PlatformFile->Initialize(&InitArgs.GetPlatformFile(), TEXT(""));
 	ReplayChanges(ManifestContent.DeletedFiles.GetFiles(), ManifestContent.ModifiedFiles.GetFiles(), ManifestContent.AddedFiles.GetFiles());
@@ -91,7 +92,8 @@ FSandboxInstance::~FSandboxInstance()
 FPersistResult FSandboxInstance::PersistSandbox(const FPersistArgs& InArgs)
 {
 	bool bManifestNeedsSave = false;
-	const FPersistResult Result = PlatformFile->PersistSandbox(InArgs,
+	const FPersistResult Result = PlatformFile->PersistSandbox(
+		SourceControlSandbox.GetProxiedSourceControl(), InArgs,
 		[this, &bManifestNeedsSave](const FSandboxedPlatformFilePath& InPath, FSandboxPlatformFile::EPersistedFileType InPersistType)
 		{
 			switch (InPersistType)
@@ -139,8 +141,13 @@ bool FSandboxInstance::DeletedPackageExistsInNonSandbox(const FString& InFilenam
 
 void FSandboxInstance::OnFileDeletionStateChanged(const FSandboxedPlatformFilePath& InFile, bool bIsDeleted)
 {
-	const bool bNeedsSave = (bIsDeleted && ManifestContent.DeletedFiles.Add(InFile.GetNonSandboxPath()))
-		|| (!bIsDeleted && ManifestContent.DeletedFiles.Remove(InFile.GetNonSandboxPath()));
+	const FNonSandboxPath NonSandboxPath = InFile.GetNonSandboxPath();
+	
+	bool bNeedsSave = false;
+	bNeedsSave |= bIsDeleted && ManifestContent.ModifiedFiles.Remove(NonSandboxPath);
+	bNeedsSave |= bIsDeleted && ManifestContent.DeletedFiles.Add(NonSandboxPath);
+	bNeedsSave |= !bIsDeleted && ManifestContent.DeletedFiles.Remove(NonSandboxPath);
+	
 	if (bNeedsSave)
 	{
 		SaveManifest();
diff --git a/Engine/Plugins/Developer/Sandbox/FileSandbox/Source/FileSandboxCore/Private/Sandbox/SandboxInstance.h b/Engine/Plugins/Developer/Sandbox/FileSandbox/Source/FileSandboxCore/Private/Sandbox/SandboxInstance.h
--- a/Engine/Plugins/Developer/Sandbox/FileSandbox/Source/FileSandboxCore/Private/Sandbox/SandboxInstance.h
+++ b/Engine/Plugins/Developer/Sandbox/FileSandbox/Source/FileSandboxCore/Private/Sandbox/SandboxInstance.h
@@ -6,12 +6,9 @@
 #include "Data/ManifestData.h"
 #include "Data/SandboxMetaData.h"
 #include "Platform/SandboxPlatformFile.h"
+#include "SourceControl/SandboxedSourceControl.h"
 #include "Templates/UnrealTemplate.h"
 
-namespace UE::FileSandboxCore
-{
-}
-
 namespace UE::FileSandboxCore
 {
 struct FSandboxInitArgs;
@@ -63,6 +60,9 @@ private:
 
 	/** Hooks into the engine's I/O system to redirect saving & loading of files. */
 	TUniquePtr<FSandboxPlatformFile> PlatformFile;
+	
+	/** Sandboxes source control operations and reroutes them to use instead. */
+	FSandboxedSourceControl SourceControlSandbox;
 
 	/** Marks the file as deleted in the manifest. */
 	void OnFileDeletionStateChanged(const FSandboxedPlatformFilePath& InFile, bool bIsDeleted);
diff --git a/Engine/Plugins/Developer/Sandbox/FileSandbox/Source/FileSandboxCore/Private/SourceControl/SandboxSourceControlProviderProxy.cpp b/Engine/Plugins/Developer/Sandbox/FileSandbox/Source/FileSandboxCore/Private/SourceControl/SandboxSourceControlProviderProxy.cpp
new file mode 100644
--- /dev/null
+++ b/Engine/Plugins/Developer/Sandbox/FileSandbox/Source/FileSandboxCore/Private/SourceControl/SandboxSourceControlProviderProxy.cpp
@@ -0,0 +1,93 @@
+﻿// Copyright Epic Games, Inc. All Rights Reserved.
+
+#include "SandboxSourceControlProviderProxy.h"
+
+#include "Features/IModularFeatures.h"
+#include "GenericPlatform/GenericPlatformFile.h"
+#include "HAL/PlatformFileManager.h"
+
+namespace UE::FileSandboxCore
+{
+namespace Private
+{
+static bool IsProhibitedOp(const FSourceControlOperationRef& InOperation)
+{
+	const FName OpName = InOperation->GetName();
+	// We're not supposed to alter source control state while in sandbox - we touch the files directly and only affect source control when persisting.
+	// Some editor operations issue a revert followed by file op, e.g. deleting an edited asset will first revert the asset, then delete it.
+	return OpName == TEXT("Revert")	
+		// This would show a popup a status update widget, don't show it in sandbox.
+		|| OpName == TEXT("UpdateStatus"); 
+}
+	
+static TOptional<ECommandResult::Type> HandleSandboxedOp(const FSourceControlOperationRef& InOperation, const TArray<FString>& InFiles)
+{
+	const FName OpName = InOperation->GetName();
+	
+	// Edits should not show up in the external source control programming until they are persisted by the sandbox
+	if (OpName == TEXT("CheckOut"))
+	{
+		return ECommandResult::Succeeded;
+	}
+	
+	// Added files should not show up in the external source control programming until they are persisted by the sandbox
+	if (OpName == TEXT("MarkForAdd"))
+	{
+		return ECommandResult::Succeeded;
+	}
+	
+	// The contract ISourceControlProvider handles deleting the file. We'll just reroute it to our sandbox. 
+	// If the changes end up being persisted, the source control changes will be routed to the real source control then.
+	if (OpName == TEXT("Delete"))
+	{
+		IPlatformFile& PlatformFile = FPlatformFileManager::Get().GetPlatformFile();
+		bool bAllSuccessful = true;
+		for (const FString& InFile : InFiles)
+		{
+			bAllSuccessful &= PlatformFile.DeleteFile(*InFile);
+		}
+		return bAllSuccessful ? ECommandResult::Succeeded : ECommandResult::Failed;
+	}
+	
+	return {};
+}
+}
+	
+FSandboxSourceControlProviderProxy::FSandboxSourceControlProviderProxy(FSandboxInstance& InInstance)
+	: Super(&ISourceControlModule::Get().GetProvider())
+	, Instance(InInstance)
+{
+	IModularFeatures::Get().RegisterModularFeature("SourceControl", this);
+}
+
+FSandboxSourceControlProviderProxy::~FSandboxSourceControlProviderProxy()
+{
+	IModularFeatures::Get().UnregisterModularFeature("SourceControl", this);
+}
+
+const FName& FSandboxSourceControlProviderProxy::GetName() const
+{
+	static FName Name(TEXT("File Sandbox"));
+	return Name;
+}
+
+ECommandResult::Type FSandboxSourceControlProviderProxy::Execute(
+	const FSourceControlOperationRef& InOperation,
+	FSourceControlChangelistPtr InChangelist, const TArray<FString>& InFiles, EConcurrency::Type InConcurrency,
+	const FSourceControlOperationComplete& InOperationCompleteDelegate
+	)
+{
+	// Handle forwarding equivalents of edit, add, and remove to sandbox.
+	if (const TOptional<ECommandResult::Type> Result = Private::HandleSandboxedOp(InOperation, InFiles))
+	{
+		return *Result;
+	}
+	
+	if (Private::IsProhibitedOp(InOperation))
+	{
+		return ECommandResult::Succeeded;
+	}
+	
+	return Super::Execute(InOperation, InChangelist, InFiles, InConcurrency, InOperationCompleteDelegate);
+}
+}
diff --git a/Engine/Plugins/Developer/Sandbox/FileSandbox/Source/FileSandboxCore/Private/SourceControl/SandboxSourceControlProviderProxy.h b/Engine/Plugins/Developer/Sandbox/FileSandbox/Source/FileSandboxCore/Private/SourceControl/SandboxSourceControlProviderProxy.h
new file mode 100644
--- /dev/null
+++ b/Engine/Plugins/Developer/Sandbox/FileSandbox/Source/FileSandboxCore/Private/SourceControl/SandboxSourceControlProviderProxy.h
@@ -0,0 +1,50 @@
+﻿// Copyright Epic Games, Inc. All Rights Reserved.
+
+#pragma once
+
+#include "HAL/Platform.h"
+#include "SourceControlProviderProxyBase.h"
+#include "Templates/UnrealTemplate.h"
+
+namespace UE::FileSandboxCore
+{
+class FSandboxInstance;
+
+/** 
+ * Intercepts any file changes the engine attempts to make and routes them FSandboxInstance.
+ * 
+ * When set-up with source control, some tools in the engine route I/O operations to the external source control tool instead of to IPlatformFile.
+ * The point of this is that the external program handles the I/O operations, e.g. "p4 delete" deletes the file and registers it with a changelist.
+ * 
+ * However, during sandbox we're not supposed to make any external changes - so we'll "swallow" those changes and re-reroute to FSandboxInstance instead.
+ * FSandboxInstance will handle re-issuing the source control commands again if the changes are persisted.
+ */
+class FSandboxSourceControlProviderProxy 
+	: public FBaseSourceControlProviderProxy
+	, public FNoncopyable
+{
+	using Super = FBaseSourceControlProviderProxy;
+public:
+	
+	explicit FSandboxSourceControlProviderProxy(FSandboxInstance& InInstance UE_LIFETIMEBOUND);
+	~FSandboxSourceControlProviderProxy();
+	
+	void SetUnderlyingSourceControlProvider(ISourceControlProvider* InProvider) { UnderlyingProvider = InProvider; }
+	
+	//~ Begin ISourceControlProvider Interface
+	virtual const FName& GetName() const override;
+	virtual ECommandResult::Type Execute(
+		const FSourceControlOperationRef& InOperation, FSourceControlChangelistPtr InChangelist, const TArray<FString>& InFiles, 
+		EConcurrency::Type InConcurrency = EConcurrency::Synchronous, 
+		const FSourceControlOperationComplete& InOperationCompleteDelegate = FSourceControlOperationComplete() 
+		) override;
+	//~ End ISourceControlProvider Interface
+	
+private:
+	
+	/** Instance of the sandbox for which we've been created. */
+	FSandboxInstance& Instance;
+};
+}
+
+
diff --git a/Engine/Plugins/Developer/Sandbox/FileSandbox/Source/FileSandboxCore/Private/SourceControl/SandboxedSourceControl.cpp b/Engine/Plugins/Developer/Sandbox/FileSandbox/Source/FileSandboxCore/Private/SourceControl/SandboxedSourceControl.cpp
new file mode 100644
--- /dev/null
+++ b/Engine/Plugins/Developer/Sandbox/FileSandbox/Source/FileSandboxCore/Private/SourceControl/SandboxedSourceControl.cpp
@@ -0,0 +1,17 @@
+﻿// Copyright Epic Games, Inc. All Rights Reserved.
+
+#include "SandboxedSourceControl.h"
+
+namespace UE::FileSandboxCore
+{
+FSandboxedSourceControl::FSandboxedSourceControl(FSandboxInstance& InSandboxInstance)
+	: SandboxInstance(InSandboxInstance)
+	, SourceControlProviderProxy(InSandboxInstance)
+	, Override(SourceControlProviderProxy)
+{
+	Override.OnProxiedProviderChanged().AddLambda([this](ISourceControlProvider& InRestoreProvider)
+	{
+		SourceControlProviderProxy.SetUnderlyingSourceControlProvider(&InRestoreProvider);
+	});
+}
+}
diff --git a/Engine/Plugins/Developer/Sandbox/FileSandbox/Source/FileSandboxCore/Private/SourceControl/SandboxedSourceControl.h b/Engine/Plugins/Developer/Sandbox/FileSandbox/Source/FileSandboxCore/Private/SourceControl/SandboxedSourceControl.h
new file mode 100644
--- /dev/null
+++ b/Engine/Plugins/Developer/Sandbox/FileSandbox/Source/FileSandboxCore/Private/SourceControl/SandboxedSourceControl.h
@@ -0,0 +1,36 @@
+﻿// Copyright Epic Games, Inc. All Rights Reserved.
+
+#pragma once
+
+#include "HAL/Platform.h"
+#include "SandboxSourceControlProviderProxy.h"
+#include "ScopedSourceControlOverride.h"
+#include "Templates/UnrealTemplate.h"
+
+namespace UE::FileSandboxCore
+{
+class FSandboxInstance;
+	
+/** Interface for FSandboxInstance to interact with source control. */
+class FSandboxedSourceControl : public FNoncopyable
+{
+public:
+
+	explicit FSandboxedSourceControl(FSandboxInstance& InSandboxInstance UE_LIFETIMEBOUND);
+	
+	/** @return The real source control provider that we're overriding. */
+	ISourceControlProvider& GetProxiedSourceControl() const { return Override.GetProxiedSourceControlProvider(); }
+	
+private:
+	
+	/** Owning sandbox instance with which source control interacts. */
+	FSandboxInstance& SandboxInstance;
+	
+	/** The sandbox proxy that should be active for as long as the sandbox exists. */
+	FSandboxSourceControlProviderProxy SourceControlProviderProxy;
+	
+	/** Sets the global source control and handles it externally changing. */
+	FScopedSourceControlOverride Override;
+};
+}
+
diff --git a/Engine/Plugins/Developer/Sandbox/FileSandbox/Source/FileSandboxCore/Private/SourceControl/ScopedSourceControlOverride.cpp b/Engine/Plugins/Developer/Sandbox/FileSandbox/Source/FileSandboxCore/Private/SourceControl/ScopedSourceControlOverride.cpp
new file mode 100644
--- /dev/null
+++ b/Engine/Plugins/Developer/Sandbox/FileSandbox/Source/FileSandboxCore/Private/SourceControl/ScopedSourceControlOverride.cpp
@@ -0,0 +1,55 @@
+﻿// Copyright Epic Games, Inc. All Rights Reserved.
+
+#include "ScopedSourceControlOverride.h"
+
+#include "ISourceControlModule.h"
+#include "ISourceControlProvider.h"
+#include "Features/IModularFeatures.h"
+
+namespace UE::FileSandboxCore
+{
+FScopedSourceControlOverride::FScopedSourceControlOverride(ISourceControlProvider& InInitialProvider)
+	: DesiredRootProvider(&InInitialProvider)
+	, ProxiedSourceControlProvider(&ISourceControlModule::Get().GetProvider())
+	, ProviderChangedHandle([this]
+	{
+		ISourceControlModule& SourceControl = ISourceControlModule::Get();
+		return SourceControl.RegisterProviderChanged(FSourceControlProviderChanged::FDelegate::CreateRaw(this, &FScopedSourceControlOverride::HandleProviderChanged));
+	}())
+{
+	SetRootProvider(*DesiredRootProvider);
+}
+
+FScopedSourceControlOverride::~FScopedSourceControlOverride()
+{
+	if (ISourceControlModule* SourceControl = ISourceControlModule::GetPtr())
+	{
+		SourceControl->UnregisterProviderChanged(ProviderChangedHandle);
+		SourceControl->SetProvider(ProxiedSourceControlProvider->GetName());
+	}
+}
+
+void FScopedSourceControlOverride::SetRootProvider(ISourceControlProvider& InNewProvider)
+{
+	const TGuardValue<bool> ReentrancyGuard(bHandlingProviderChanges, true);
+	DesiredRootProvider = &InNewProvider;
+	ISourceControlModule& SourceControl = ISourceControlModule::Get();
+	SourceControl.SetProvider(DesiredRootProvider->GetName());
+}
+
+void FScopedSourceControlOverride::HandleProviderChanged(ISourceControlProvider& OldProvider, ISourceControlProvider& NewProvider)
+{
+	// If we are currently setting ourselves as the provider
+	if (bHandlingProviderChanges)
+	{
+		return;
+	}
+	const TGuardValue<bool> ReentrancyGuard(bHandlingProviderChanges, true);
+	
+	// if we receive this event we should be installed as the current provider, if we aren't resetting ourselves
+	check(&OldProvider == DesiredRootProvider);
+	ProxiedSourceControlProvider = &NewProvider;
+	OnProxiedProviderChangedDelegate.Broadcast(NewProvider);
+	ISourceControlModule::Get().SetProvider(DesiredRootProvider->GetName());
+}
+}
diff --git a/Engine/Plugins/Developer/Sandbox/FileSandbox/Source/FileSandboxCore/Private/SourceControl/ScopedSourceControlOverride.h b/Engine/Plugins/Developer/Sandbox/FileSandbox/Source/FileSandboxCore/Private/SourceControl/ScopedSourceControlOverride.h
new file mode 100644
--- /dev/null
+++ b/Engine/Plugins/Developer/Sandbox/FileSandbox/Source/FileSandboxCore/Private/SourceControl/ScopedSourceControlOverride.h
@@ -0,0 +1,54 @@
+﻿// Copyright Epic Games, Inc. All Rights Reserved.
+
+#pragma once
+
+#include "Delegates/Delegate.h"
+#include "Delegates/IDelegateInstance.h"
+#include "HAL/Platform.h"
+#include "Misc/NotNull.h"
+#include "Templates/UnrealTemplate.h"
+
+class ISourceControlProvider;
+
+namespace UE::FileSandboxCore
+{
+DECLARE_MULTICAST_DELEGATE_OneParam(FSourceControlProviderEvent, ISourceControlProvider&);
+	
+/** Forces the global ISourceControlProvider to always be the passed in one, even when the user manually changes the source control provider. */
+class FScopedSourceControlOverride : public FNoncopyable
+{
+public:
+	
+	explicit FScopedSourceControlOverride(ISourceControlProvider& InInitialProvider UE_LIFETIMEBOUND);
+	~FScopedSourceControlOverride();
+	
+	/** Changes the root source control provider again, as if ~FScopedSourceControlOverride() and FScopedSourceControlOverride were called again. */
+	void SetRootProvider(ISourceControlProvider& InNewProvider);
+	
+	/** @return The real source control provider that will be restored to. */
+	ISourceControlProvider& GetProxiedSourceControlProvider() const { return *ProxiedSourceControlProvider; }
+	
+	/** Invoked when the underlying source control provider changes, e.g. when the user changes the source control settings. */
+	FSourceControlProviderEvent& OnProxiedProviderChanged() { return OnProxiedProviderChangedDelegate; }
+
+private:
+	
+	/** The source control provider that should remain the root source control provider. */
+	TNotNull<ISourceControlProvider*> DesiredRootProvider;
+	/** The provider to restore to when the scope ends. */
+	TNotNull<ISourceControlProvider*> ProxiedSourceControlProvider;
+	
+	/** Reentry guard when handling provider changes. */
+	bool bHandlingProviderChanges = false;
+	
+	/** Delegate handle for provider changes. */
+	const FDelegateHandle ProviderChangedHandle;
+	
+	/** Invoked when the underlying source control provider changes */
+	FSourceControlProviderEvent OnProxiedProviderChangedDelegate;
+	
+	/** Delegate to handle provider change and change our underlying provider. */
+	void HandleProviderChanged(ISourceControlProvider& OldProvider, ISourceControlProvider& NewProvider);
+};
+}
+
diff --git a/Engine/Plugins/Developer/Sandbox/FileSandbox/Source/FileSandboxCore/Private/SourceControl/SourceControlProviderProxyBase.h b/Engine/Plugins/Developer/Sandbox/FileSandbox/Source/FileSandboxCore/Private/SourceControl/SourceControlProviderProxyBase.h
new file mode 100644
--- /dev/null
+++ b/Engine/Plugins/Developer/Sandbox/FileSandbox/Source/FileSandboxCore/Private/SourceControl/SourceControlProviderProxyBase.h
@@ -0,0 +1,173 @@
+﻿// Copyright Epic Games, Inc. All Rights Reserved.
+
+#pragma once
+
+#include "ISourceControlModule.h"
+#include "ISourceControlProvider.h"
+#if SOURCE_CONTROL_WITH_SLATE
+#include "Widgets/SNullWidget.h"
+#endif
+
+#define WRAP(FuncName, ...) if (ISourceControlProvider* Provider = GetUnderlyingProvider()) { Provider->FuncName(__VA_ARGS__); }
+#define WRAP_RET(Default, FuncName, ...) if (ISourceControlProvider* Provider = GetUnderlyingProvider()) { return Provider->FuncName(__VA_ARGS__); } return Default; 
+
+namespace UE::FileSandboxCore
+{
+/** 
+ * Forwards all calls to an optional, proxied source control provider. 
+ * Acts as base class to reduce boilerplate when proxying source control classes.
+ */
+class FBaseSourceControlProviderProxy : public ISourceControlProvider
+{
+public:
+	
+	explicit FBaseSourceControlProviderProxy(ISourceControlProvider* InActualProvider)
+		: UnderlyingProvider(InActualProvider)
+	{}
+	
+	//~ Begin ISourceControlProvider Interface
+	virtual void Init(bool bForceConnection = true) override
+	{
+		WRAP(Init, bForceConnection)
+	}
+	virtual void Close() override
+	{
+		WRAP(Close);
+	}
+	virtual FText GetStatusText() const override
+	{
+		WRAP_RET(FText::GetEmpty(), GetStatusText)
+	}
+	virtual TMap<EStatus, FString> GetStatus() const override
+	{
+		WRAP_RET({}, GetStatus)
+	}
+	virtual bool IsEnabled() const override
+	{
+		WRAP_RET(true, IsEnabled)
+	}
+	virtual bool IsAvailable() const override
+	{
+		WRAP_RET(true, IsAvailable)
+	}
+	virtual bool QueryStateBranchConfig(const FString& ConfigSrc, const FString& ConfigDest) override
+	{
+		WRAP_RET(false, QueryStateBranchConfig, ConfigSrc, ConfigDest)
+	}
+	virtual void RegisterStateBranches(const TArray<FString>& BranchNames, const FString& ContentRoot) override
+	{
+		WRAP(RegisterStateBranches, BranchNames, ContentRoot)
+	}
+	virtual int32 GetStateBranchIndex(const FString& BranchName) const override
+	{
+		WRAP_RET(INDEX_NONE, GetStateBranchIndex, BranchName);
+	}
+	virtual bool GetStateBranchAtIndex(int32 BranchIndex, FString& OutBranchName) const override
+	{
+		WRAP_RET(false, GetStateBranchAtIndex, BranchIndex, OutBranchName);
+	}
+	virtual ECommandResult::Type GetState(const TArray<FString>& InFiles, TArray<FSourceControlStateRef>& OutState, EStateCacheUsage::Type InStateCacheUsage) override
+	{
+		WRAP_RET(ECommandResult::Failed, GetState, InFiles, OutState, InStateCacheUsage)
+	}
+	virtual ECommandResult::Type GetState(const TArray<FSourceControlChangelistRef>& InChangelists, TArray<FSourceControlChangelistStateRef>& OutState, EStateCacheUsage::Type InStateCacheUsage) override
+	{
+		WRAP_RET(ECommandResult::Failed, GetState, InChangelists, OutState, InStateCacheUsage)
+	}
+	virtual TArray<FSourceControlStateRef> GetCachedStateByPredicate(TFunctionRef<bool(const FSourceControlStateRef&)> Predicate) const override
+	{
+		WRAP_RET({}, GetCachedStateByPredicate, Predicate)
+	}
+	virtual FDelegateHandle RegisterSourceControlStateChanged_Handle(const FSourceControlStateChanged::FDelegate& SourceControlStateChanged) override
+	{
+		WRAP_RET({}, RegisterSourceControlStateChanged_Handle, SourceControlStateChanged)
+	}
+	virtual void UnregisterSourceControlStateChanged_Handle(FDelegateHandle Handle) override
+	{
+		WRAP(UnregisterSourceControlStateChanged_Handle, Handle)
+	}
+	virtual ECommandResult::Type Execute(const FSourceControlOperationRef& InOperation, FSourceControlChangelistPtr InChangelist, const TArray<FString>& InFiles, EConcurrency::Type InConcurrency = EConcurrency::Synchronous, const FSourceControlOperationComplete& InOperationCompleteDelegate = FSourceControlOperationComplete()) override
+	{
+		WRAP_RET(ECommandResult::Failed, Execute, InOperation, InChangelist, InFiles, InConcurrency, InOperationCompleteDelegate)
+	} 
+	virtual bool CanExecuteOperation(const FSourceControlOperationRef& InOperation) const override
+	{
+		WRAP_RET(false, CanExecuteOperation, InOperation)
+	}
+	virtual bool CanCancelOperation(const FSourceControlOperationRef& InOperation) const override
+	{
+		WRAP_RET(false, CanCancelOperation, InOperation)
+	}
+	virtual void CancelOperation(const FSourceControlOperationRef& InOperation) override
+	{
+		WRAP(CancelOperation, InOperation)
+	}
+	virtual TArray<TSharedRef<class ISourceControlLabel>> GetLabels(const FString& InMatchingSpec) const override
+	{
+		WRAP_RET({}, GetLabels, InMatchingSpec)
+	}
+	virtual TArray<FSourceControlChangelistRef> GetChangelists(EStateCacheUsage::Type InStateCacheUsage) override
+	{
+		WRAP_RET({}, GetChangelists, InStateCacheUsage)
+	}
+	virtual bool UsesLocalReadOnlyState() const override
+	{
+		WRAP_RET(false, UsesLocalReadOnlyState)
+	}
+	virtual bool UsesChangelists() const override
+	{
+		WRAP_RET(false, UsesChangelists)
+	}
+	virtual bool UsesUncontrolledChangelists() const override
+	{
+		WRAP_RET(false, UsesUncontrolledChangelists)
+	}
+	virtual bool UsesCheckout() const override
+	{
+		WRAP_RET(false, UsesCheckout)
+	}
+	virtual bool UsesFileRevisions() const override
+	{
+		WRAP_RET(false, UsesFileRevisions)
+	}
+	virtual bool UsesSnapshots() const override
+	{
+		WRAP_RET(false, UsesSnapshots)
+	}
+	virtual bool AllowsDiffAgainstDepot() const override
+	{
+		WRAP_RET(false, AllowsDiffAgainstDepot)
+	}
+	virtual TOptional<bool> IsAtLatestRevision() const override
+	{
+		WRAP_RET({}, IsAtLatestRevision)
+	}
+	virtual TOptional<int> GetNumLocalChanges() const override
+	{
+		WRAP_RET({}, GetNumLocalChanges)
+	}
+	virtual void Tick() override
+	{
+		WRAP(Tick)
+	}
+#if SOURCE_CONTROL_WITH_SLATE
+	virtual TSharedRef<class SWidget> MakeSettingsWidget() const override
+	{
+		WRAP_RET(SNullWidget::NullWidget, MakeSettingsWidget)
+	}
+#endif
+	//~ End ISourceControlProvider Interface
+	
+protected:
+	
+	ISourceControlProvider* UnderlyingProvider;
+	
+	ISourceControlProvider* GetUnderlyingProvider() const 
+	{
+		return UnderlyingProvider;
+	}
+};
+}
+
+#undef WRAP
+#undef WRAP_RET
\ No newline at end of file
