File sandbox: Refactor usage of directory watchers.

All dependencies to DirectoryWatcher should now be wrapped with UE_SANDBOX_WITH_DIRECTORY_WATCHER, which we set by default when building the editor.

Some internal tools are not supposed to build with DirectoryWatcher in runtime configuration. Hence we need to compile it out optionally.

#Synced-CL 47637010

diff --git a/Engine/Plugins/Developer/Sandbox/FileSandbox/Source/FileSandboxCore/FileSandboxCore.Build.cs b/Engine/Plugins/Developer/Sandbox/FileSandbox/Source/FileSandboxCore/FileSandboxCore.Build.cs
--- a/Engine/Plugins/Developer/Sandbox/FileSandbox/Source/FileSandboxCore/FileSandboxCore.Build.cs
+++ b/Engine/Plugins/Developer/Sandbox/FileSandbox/Source/FileSandboxCore/FileSandboxCore.Build.cs
@@ -16,7 +16,6 @@ namespace UnrealBuildTool.Rules
 			PrivateDependencyModuleNames.AddRange(
 				new string[]
 				{
-					"DirectoryWatcher",
 					"Engine",
 					"Json",
 					"JsonUtilities", 
@@ -34,6 +33,10 @@ namespace UnrealBuildTool.Rules
 
 			if (Target.bBuildEditor)
 			{
+				// DirectoryWatcher module is a disallowed dependency for some internal tools that use a runtime version of this build. 
+				PrivateDefinitions.Add("UE_SANDBOX_WITH_DIRECTORY_WATCHER=1");
+				PrivateDependencyModuleNames.Add("DirectoryWatcher");
+				
 				PrivateDependencyModuleNames.AddRange(
 					new string[]
 					{
@@ -42,6 +45,10 @@ namespace UnrealBuildTool.Rules
 					}
 				);
 			}
+			else
+			{
+				PrivateDefinitions.Add("UE_SANDBOX_WITH_DIRECTORY_WATCHER=0");
+			}
 		}
 	}
 }
diff --git a/Engine/Plugins/Developer/Sandbox/FileSandbox/Source/FileSandboxCore/Private/Sandbox/Platform/SandboxMountPoint.h b/Engine/Plugins/Developer/Sandbox/FileSandbox/Source/FileSandboxCore/Private/Sandbox/Platform/SandboxMountPoint.h
new file mode 100644
--- /dev/null
+++ b/Engine/Plugins/Developer/Sandbox/FileSandbox/Source/FileSandboxCore/Private/Sandbox/Platform/SandboxMountPoint.h
@@ -0,0 +1,18 @@
+﻿// Copyright Epic Games, Inc. All Rights Reserved.
+
+#pragma once
+
+#include "SandboxedPlatformFilePath.h"
+#include "Delegates/IDelegateInstance.h"
+
+namespace UE::FileSandboxCore
+{
+struct FSandboxMountPoint
+{
+	/** Sandbox path */
+	FSandboxedPlatformFilePath Path;
+
+	/** Sandbox directory watcher delegate handle (if any) */
+	FDelegateHandle OnDirectoryChangedHandle;
+};
+}
diff --git a/Engine/Plugins/Developer/Sandbox/FileSandbox/Source/FileSandboxCore/Private/Sandbox/Platform/SandboxPlatformFile.cpp b/Engine/Plugins/Developer/Sandbox/FileSandbox/Source/FileSandboxCore/Private/Sandbox/Platform/SandboxPlatformFile.cpp
--- a/Engine/Plugins/Developer/Sandbox/FileSandbox/Source/FileSandboxCore/Private/Sandbox/Platform/SandboxPlatformFile.cpp
+++ b/Engine/Plugins/Developer/Sandbox/FileSandbox/Source/FileSandboxCore/Private/Sandbox/Platform/SandboxPlatformFile.cpp
@@ -4,31 +4,28 @@
 
 #include "Containers/UnrealString.h"
 #include "Engine/GameEngine.h"
-#include "Engine/Level.h"
 #include "Engine/World.h"
 #include "GenericPlatform/GenericPlatformFile.h"
 #include "HAL/PlatformFileManager.h"
-#include "ISourceControlModule.h"
 #include "ISourceControlProvider.h"
 #include "ISourceControlState.h"
 #include "JsonObjectConverter.h"
 #include "LogFileSandbox.h"
 #include "Logging/LogVerbosity.h"
 #include "Misc/FileHelper.h"
-#include "Misc/MessageDialog.h"
 #include "Misc/PackageName.h"
 #include "Misc/Paths.h"
 #include "Misc/ScopeLock.h"
 #include "Modules/ModuleManager.h"
 #include "RenderingThread.h"
+#include "SandboxMountPoint.h"
 #include "SourceControlOperations.h"
-#include "Templates/UniquePtr.h"
 #include "Types/Sandbox/PersistArgs.h"
 #include "Types/Sandbox/PersistResult.h"
-#include "UObject/Linker.h"
 #include "UObject/LinkerLoad.h"
 #include "UObject/Package.h"
-#include "Utils/PackageSandboxUtils.h"
+#include "Utils/FileChange/FileChange.h"
+#include "Utils/Persist/DiscardUtils.h"
 #include "Utils/SandboxUtils.h"
 
 #if WITH_EDITOR
@@ -36,12 +33,8 @@
 #include "Editor.h"
 #include "Editor/UnrealEdEngine.h"
 #include "FileHelpers.h"
-#include "GameMapsSettings.h"
 #include "IDirectoryWatcher.h"
-#include "PackageTools.h"
 #include "Selection.h"
-#include "Subsystems/AssetEditorSubsystem.h"
-#include "Subsystems/EditorActorSubsystem.h"
 #endif
 
 class IDirectoryWatcher;
@@ -847,101 +840,38 @@ FPersistResult FSandboxPlatformFile::PersistSandbox(ISourceControlProvider& InSo
 }
 
 void FSandboxPlatformFile::DiscardSandbox(
-	TArray<FName>& OutPackagesPendingHotReload, TArray<FName>& OutPackagesPendingPurge, TArray<FFileChangeData>* OutFileChanges
+	TArray<FName>& OutPackagesPendingHotReload, 
+	TArray<FName>& OutPackagesPendingPurge,
+	TFunctionRef<void(const FFileChange&)> InProcessFileChangeCallback
 	)
 {
 	// We need to disable the sandbox while we do this
 	TGuardValue<TAtomic<bool>, bool> DisableSandboxGuard(bSandboxEnabled, false);
-
+	
+#if UE_SANDBOX_WITH_DIRECTORY_WATCHER
 	TArray<FFileChangeData> FileChanges;
+#endif
+	const auto HandleFileChange = [&](const FFileChange& Change)
+	{
+#if UE_SANDBOX_WITH_DIRECTORY_WATCHER
+		FileChanges.Emplace(Change.ToDirectoryWatcherChange());
+#endif
+		InProcessFileChangeCallback(Change);
+	};
 
 	// Add any files that were deleted by the sandbox but exist in the non-sandbox directory
 	{
 		FScopeLock DeletedSandboxPathsLock(&DeletedSandboxPathsCS);
-		for (auto It = DeletedSandboxPaths.CreateIterator(); It; ++It)
-		{
-			const FNonSandboxPath& NonSandboxPath = *It;
-			
-			// If this file maps to a package then we need to flush its linker so that we can remove the file from the sandbox
-			FName PackageName;
-			FlushPackageFile(NonSandboxPath, &PackageName);
-
-			if (LowerLevel->FileExists(*NonSandboxPath))
-			{
-				if (!PackageName.IsNone())
-				{
-					OutPackagesPendingHotReload.Add(PackageName);
-				}
-				FileChanges.Emplace(NonSandboxPath, FFileChangeData::FCA_Added);
-			}
-			It.RemoveCurrent();
-			continue;
-		}
-
-		// Clear the deleted path information
+		DiscardRemovedPaths(DeletedSandboxPaths, *LowerLevel, OutPackagesPendingPurge, HandleFileChange);
 		DeletedSandboxPaths.Reset();
 	}
 
 	{
 		FScopeLock SandboxMountPointsLock(&SandboxMountPointsCS);
-		for (const FSandboxMountPoint& SandboxMountPoint : SandboxMountPoints)
-		{
-			// Modify any files that exist in both the non-sandbox and sandbox directories
-			// Delete any files that exist in the sandbox but don't exist in the non-sandbox directory
-			LowerLevel->IterateDirectoryRecursively(*SandboxMountPoint.Path.GetSandboxPath(), [this, &SandboxMountPoint, &FileChanges, &OutPackagesPendingHotReload, &OutPackagesPendingPurge](const TCHAR* InFilenameOrDirectory, const bool InIsDirectory) -> bool
-			{
-				if (!InIsDirectory)
-				{
-					const FSandboxedPlatformFilePath RemappedFilePath = FSandboxedPlatformFilePath::CreateNonSandboxPath(FPaths::ConvertRelativePathToFull(InFilenameOrDirectory), SandboxMountPoint.Path);
-
-					// If this file maps to a package then we need to flush its linker so that we can remove the file from the sandbox
-					FName PackageName;
-					FlushPackageFile(RemappedFilePath.GetNonSandboxPath(), &PackageName, false);
-
-					if (LowerLevel->FileExists(*RemappedFilePath.GetNonSandboxPath()))
-					{
-						if (!PackageName.IsNone())
-						{
-							OutPackagesPendingHotReload.Add(PackageName);
-							OutPackagesPendingPurge.Remove(PackageName);
-						}
-						
-						FileChanges.Emplace(RemappedFilePath.GetNonSandboxPath(), FFileChangeData::FCA_Modified);
-					}
-					else
-					{
-						if (!PackageName.IsNone())
-						{
-							OutPackagesPendingPurge.Add(PackageName);
-							OutPackagesPendingHotReload.Remove(PackageName);
-						}
-						
-						FileChanges.Emplace(RemappedFilePath.GetNonSandboxPath(), FFileChangeData::FCA_Removed);
-					}
-				}
-				return true; // Continue iteration
-			});
-
-			//FString SandboxRootPath = GetSandboxRootPath();
-			//FString Stashed = SandboxRootPath.Replace(TEXT("Sandbox"), TEXT("Stashed"));
-			//LowerLevel->MoveFile(SandboxRootPath, Stashed);
-			// Delete everything under the mount point
-			//LowerLevel->IterateDirectory(*SandboxMountPoint.Path.GetSandboxPath(), [this](const TCHAR* InFilenameOrDirectory, const bool InIsDirectory) -> bool
-			//{
-			//	if (InIsDirectory)
-			//	{
-			//		LowerLevel->DeleteDirectoryRecursively(InFilenameOrDirectory);
-			//	}
-			//	else
-			//	{
-			//		LowerLevel->DeleteFile(InFilenameOrDirectory);
-			//	}
-			//	return true; // Continue iteration
-			//});
-		}
+		DiscardAddedAndModifiedPaths(SandboxMountPoints, *LowerLevel, OutPackagesPendingHotReload, OutPackagesPendingPurge, HandleFileChange);
 	}
 
-#if WITH_EDITOR
+#if UE_SANDBOX_WITH_DIRECTORY_WATCHER
 	// Notify that the sandboxed directories have been restored to their original state
 	if (FDirectoryWatcherModule* DirectoryWatcherModule = GetDirectoryWatcherModuleIfLoaded())
 	{
@@ -959,11 +889,6 @@ void FSandboxPlatformFile::DiscardSandbox(
 		}
 	}
 #endif
-
-	if (OutFileChanges)
-	{
-		*OutFileChanges = MoveTemp(FileChanges);
-	}
 }
 
 TArray<FString> FSandboxPlatformFile::GatherSandboxChangedFilenames() const
@@ -1045,7 +970,7 @@ void FSandboxPlatformFile::RegisterContentMountPath(const FString& InAssetPath,
 		FSandboxMountPoint& SandboxMountPoint = SandboxMountPoints.Add_GetRef(
 			FSandboxMountPoint{ FSandboxedPlatformFilePath(MoveTemp(AbsoluteNonSandboxPath), MoveTemp(AbsoluteSandboxPath)) }
 			);
-#if WITH_EDITOR
+#if UE_SANDBOX_WITH_DIRECTORY_WATCHER
 		if (IDirectoryWatcher* DirectoryWatcher = GetDirectoryWatcher())
 		{
 			DirectoryWatcher->RegisterDirectoryChangedCallback_Handle(
@@ -1079,7 +1004,7 @@ void FSandboxPlatformFile::UnregisterContentMountPath(const FString& InAssetPath
 		SandboxMountPoints.RemoveAll([&AbsoluteSandboxPath](FSandboxMountPoint& InSandboxMountPoint) -> bool
 		{
 			const bool bShouldRemove = InSandboxMountPoint.Path.GetSandboxPath() == AbsoluteSandboxPath;
-#if WITH_EDITOR
+#if UE_SANDBOX_WITH_DIRECTORY_WATCHER
 			if (bShouldRemove && InSandboxMountPoint.OnDirectoryChangedHandle.IsValid())
 			{
 				if (IDirectoryWatcher* DirectoryWatcher = GetDirectoryWatcherIfLoaded())
@@ -1189,7 +1114,7 @@ void FSandboxPlatformFile::NotifyFileDeleted(const FSandboxedPlatformFilePath& I
 		return;
 	}
 
-#if WITH_EDITOR
+#if UE_SANDBOX_WITH_DIRECTORY_WATCHER
 	if (FDirectoryWatcherModule* DirectoryWatcherModule = GetDirectoryWatcherModuleIfLoaded())
 	{
 		FFileChangeData FileChange(InPath.GetNonSandboxPath(), FFileChangeData::FCA_Removed);
@@ -1298,8 +1223,7 @@ TArray<FSandboxPlatformFile::FDirectoryItem> FSandboxPlatformFile::GetDirectoryC
 	return DirectoryContents;
 }
 
-#if WITH_EDITOR
-
+#if UE_SANDBOX_WITH_DIRECTORY_WATCHER
 void FSandboxPlatformFile::OnSandboxDirectoryChanged(const TArray<FFileChangeData>& FileChanges, FSandboxedPlatformFilePath MountPath)
 {
 	if (!IsSandboxEnabled())
@@ -1330,6 +1254,5 @@ void FSandboxPlatformFile::OnSandboxDirectoryChanged(const TArray<FFileChangeDat
 		DirectoryWatcherModule->RegisterExternalChanges(RemappedFileChanges);
 	}
 }
-
 #endif
 }
\ No newline at end of file
diff --git a/Engine/Plugins/Developer/Sandbox/FileSandbox/Source/FileSandboxCore/Private/Sandbox/Platform/SandboxPlatformFile.h b/Engine/Plugins/Developer/Sandbox/FileSandbox/Source/FileSandboxCore/Private/Sandbox/Platform/SandboxPlatformFile.h
--- a/Engine/Plugins/Developer/Sandbox/FileSandbox/Source/FileSandboxCore/Private/Sandbox/Platform/SandboxPlatformFile.h
+++ b/Engine/Plugins/Developer/Sandbox/FileSandbox/Source/FileSandboxCore/Private/Sandbox/Platform/SandboxPlatformFile.h
@@ -4,16 +4,23 @@
 
 #include "Containers/Map.h"
 #include "Containers/Set.h"
+#include "Delegates/Delegate.h"
 #include "Delegates/IDelegateInstance.h"
 #include "GenericPlatform/GenericPlatformFile.h"
-#include "IDirectoryWatcher.h"
 #include "SandboxedPlatformFilePath.h"
 
+#if UE_SANDBOX_WITH_DIRECTORY_WATCHER
+#include "IDirectoryWatcher.h"
+#endif
+
 class ISourceControlProvider;
 
 namespace UE::FileSandboxCore
 {
-struct FPersistResult;
+	struct FFileChange;
+	struct FDiscardSandboxArgs;
+	struct FSandboxMountPoint;
+	struct FPersistResult;
 struct FPersistArgs;
 	
 DECLARE_MULTICAST_DELEGATE_TwoParams(FDeletionStateChanged, const FSandboxedPlatformFilePath&, bool /* bIsDeleted */);
@@ -77,11 +84,12 @@ public:
 	 *	Discard the sandbox state
 	 *	This will trigger directory watcher notifications for files that are restored.
 	 *	This will also gather packages that need to be hot reloaded or purge from memory
-	 *	@param OutPackagePendingHotReload	Array to be filled with package name that need to be hot reloaded
-	 *	@param OutPackagesPendingPurge		Array to be filled with package name that need to be purged
-	 *	@param OutFileChanges				Array to be filled with changes made to the filesystem
 	 */
-	void DiscardSandbox(TArray<FName>& OutPackagesPendingHotReload, TArray<FName>& OutPackagesPendingPurge, TArray<FFileChangeData>* OutFileChanges = nullptr);
+	void DiscardSandbox(
+		TArray<FName>& OutPackagesPendingHotReload,
+		TArray<FName>& OutPackagesPendingPurge,
+		TFunctionRef<void(const FFileChange&)> InProcessFileChangeCallback = [](const FFileChange&){}
+		);
 
 	/** 
 	 * Gather all files changes that are currently in the sandbox.
@@ -105,15 +113,6 @@ private:
 		FString Path;
 		FFileStatData StatData;
 	};
-
-	struct FSandboxMountPoint
-	{
-		/** Sandbox path */
-		FSandboxedPlatformFilePath Path;
-
-		/** Sandbox directory watcher delegate handle (if any) */
-		FDelegateHandle OnDirectoryChangedHandle;
-	};
 	
 	/** Callback when a new content path is mounted */
 	void OnContentPathMounted(const FString& InAssetPath, const FString& InFilesystemPath);
@@ -150,7 +149,7 @@ private:
 	/** Helper function to get the contents of a directory, taking into account the sandbox state - paths are returned relative to InDirBase */
 	TArray<FDirectoryItem> GetDirectoryContents(const FSandboxedPlatformFilePath& InPath, const TCHAR* InDirBase) const;
 
-#if WITH_EDITOR
+#if UE_SANDBOX_WITH_DIRECTORY_WATCHER
 	/** Called when a file in a sandbox directory changes on disk */
 	void OnSandboxDirectoryChanged(const TArray<FFileChangeData>& FileChanges, FSandboxedPlatformFilePath MountPath);
 #endif
diff --git a/Engine/Plugins/Developer/Sandbox/FileSandbox/Source/FileSandboxCore/Private/Sandbox/SandboxInstance.cpp b/Engine/Plugins/Developer/Sandbox/FileSandbox/Source/FileSandboxCore/Private/Sandbox/SandboxInstance.cpp
--- a/Engine/Plugins/Developer/Sandbox/FileSandbox/Source/FileSandboxCore/Private/Sandbox/SandboxInstance.cpp
+++ b/Engine/Plugins/Developer/Sandbox/FileSandbox/Source/FileSandboxCore/Private/Sandbox/SandboxInstance.cpp
@@ -11,24 +11,20 @@
 #include "Utils/FileSandboxUtils.h"
 #include "Utils/SandboxDirectoryUtils.h"
 #include "Utils/SandboxUtils.h"
+#include "Utils/FileChange/FileChange.h"
 
 namespace UE::FileSandboxCore
 {
 namespace Private
 {
-static void SetManifestFileFromRevertedFileChanges(FFileSandboxCore_ManifestData& OutManifestData, TConstArrayView<FFileChangeData> InFileChanges)
+static void TrackRevertedFileChangeInManifest(FFileSandboxCore_ManifestData& OutManifestData, const FFileChange& InChange)
 {
-	OutManifestData.Empty();
-
-	for (const FFileChangeData& Change : InFileChanges)
+	switch (InChange.Action)
 	{
-		switch (Change.Action)
-		{
-		case FFileChangeData::FCA_Added: OutManifestData.DeletedFiles.Add(Change.Filename); break;
-		case FFileChangeData::FCA_Modified: OutManifestData.ModifiedFiles.Add(Change.Filename); break;
-		case FFileChangeData::FCA_Removed: OutManifestData.AddedFiles.Add(Change.Filename); break;
-		default: break;
-		}
+	case EFileChangeAction::Added: OutManifestData.DeletedFiles.Add(InChange.Filename); break;
+	case EFileChangeAction::Modified: OutManifestData.ModifiedFiles.Add(InChange.Filename); break;
+	case EFileChangeAction::Removed: OutManifestData.AddedFiles.Add(InChange.Filename); break;
+	default: break;
 	}
 }
 }
@@ -67,13 +63,18 @@ FSandboxInstance::FSandboxInstance(
 	
 FSandboxInstance::~FSandboxInstance()
 {
+	// TODO DP UE-353328: In the future, we should save these changes out as the happen. Otherwise, we cannot recover from crashes that happen during sandbox.
+	ManifestContent.Empty();
+	const auto ProcessRevertedChange = [this](const FFileChange& Change)
+	{
+		Private::TrackRevertedFileChangeInManifest(ManifestContent, Change); 
+	};
+	
 	TArray<FName> PackagesPendingHotReload, PackagesPendingPurge;
-	TArray<FFileChangeData> FileChanges;
-	PlatformFile->DiscardSandbox(PackagesPendingHotReload, PackagesPendingPurge, &FileChanges);
+	PlatformFile->DiscardSandbox(PackagesPendingHotReload, PackagesPendingPurge, ProcessRevertedChange);
 	PlatformFile.Reset();
 
 	// Save the data before hot reloading or purging: in case that causes a crash we'll have an update to manifest file. 
-	Private::SetManifestFileFromRevertedFileChanges(ManifestContent, FileChanges);
 	SaveManifest();
 	// This just exists to handle e.g. the user having deleted the file manually during the sandbox. Normally, the metadata file should still be there.
 	if (!FPaths::FileExists(RootDirectory / GetMetadataFileName()))
diff --git a/Engine/Plugins/Developer/Sandbox/FileSandbox/Source/FileSandboxCore/Private/Utils/FileChange/EFileChangeAction.h b/Engine/Plugins/Developer/Sandbox/FileSandbox/Source/FileSandboxCore/Private/Utils/FileChange/EFileChangeAction.h
new file mode 100644
--- /dev/null
+++ b/Engine/Plugins/Developer/Sandbox/FileSandbox/Source/FileSandboxCore/Private/Utils/FileChange/EFileChangeAction.h
@@ -0,0 +1,39 @@
+﻿// Copyright Epic Games, Inc. All Rights Reserved.
+
+#pragma once
+
+#include "HAL/Platform.h"
+#if UE_SANDBOX_WITH_DIRECTORY_WATCHER
+#include "IDirectoryWatcher.h"
+#endif
+
+namespace UE::FileSandboxCore
+{
+enum class EFileChangeAction : uint8
+{
+	Added,
+	Modified,
+	Removed
+};
+
+#if UE_SANDBOX_WITH_DIRECTORY_WATCHER
+/** @return Converts the change action to the equivalent IDirectoryWatcher type. */
+FFileChangeData::EFileChangeAction ToDirectoryWatcherAction(EFileChangeAction InAction);
+#endif
+}
+
+namespace UE::FileSandboxCore
+{
+#if UE_SANDBOX_WITH_DIRECTORY_WATCHER
+inline FFileChangeData::EFileChangeAction ToDirectoryWatcherAction(EFileChangeAction InAction)
+{
+	switch (InAction)
+	{
+	case EFileChangeAction::Added: return FFileChangeData::FCA_Added;
+	case EFileChangeAction::Modified: return FFileChangeData::FCA_Modified;
+	case EFileChangeAction::Removed: return FFileChangeData::FCA_Removed;
+	default: checkNoEntry(); return FFileChangeData::FCA_Unknown;
+	}
+}
+#endif
+}
\ No newline at end of file
diff --git a/Engine/Plugins/Developer/Sandbox/FileSandbox/Source/FileSandboxCore/Private/Utils/FileChange/FileChange.h b/Engine/Plugins/Developer/Sandbox/FileSandbox/Source/FileSandboxCore/Private/Utils/FileChange/FileChange.h
new file mode 100644
--- /dev/null
+++ b/Engine/Plugins/Developer/Sandbox/FileSandbox/Source/FileSandboxCore/Private/Utils/FileChange/FileChange.h
@@ -0,0 +1,32 @@
+﻿// Copyright Epic Games, Inc. All Rights Reserved.
+
+#pragma once
+
+#include "Containers/UnrealString.h"
+#include "EFileChangeAction.h"
+#include "Misc/Paths.h"
+
+namespace UE::FileSandboxCore
+{
+/** Describes a file change. This mirrors FFileChangeData from DirectoryWatchers because that module cannot be built in all runtime configurations. */
+struct FFileChange
+{
+	FString Filename;
+	EFileChangeAction Action;
+	
+	FFileChange(const FString& InFilename, EFileChangeAction InAction)
+		: Filename(InFilename)
+		, Action(InAction)
+	{
+		FPaths::MakeStandardFilename(Filename);
+	}
+	
+#if UE_SANDBOX_WITH_DIRECTORY_WATCHER
+	/** @return Converts this change to the equivalent IDirectoryWatcher type. */
+	FFileChangeData ToDirectoryWatcherChange() const
+	{
+		return FFileChangeData(Filename, ToDirectoryWatcherAction(Action));
+	}
+#endif
+};
+}
diff --git a/Engine/Plugins/Developer/Sandbox/FileSandbox/Source/FileSandboxCore/Private/Utils/Persist/DiscardUtils.cpp b/Engine/Plugins/Developer/Sandbox/FileSandbox/Source/FileSandboxCore/Private/Utils/Persist/DiscardUtils.cpp
new file mode 100644
--- /dev/null
+++ b/Engine/Plugins/Developer/Sandbox/FileSandbox/Source/FileSandboxCore/Private/Utils/Persist/DiscardUtils.cpp
@@ -0,0 +1,106 @@
+﻿// Copyright Epic Games, Inc. All Rights Reserved.
+
+#include "DiscardUtils.h"
+
+#include "Containers/Array.h"
+#include "Containers/ArrayView.h"
+#include "Containers/Set.h"
+#include "Containers/UnrealString.h"
+#include "GenericPlatform/GenericPlatformFile.h"
+#include "Sandbox/Platform/SandboxMountPoint.h"
+#include "UObject/NameTypes.h"
+#include "Utils/FileChange/FileChange.h"
+#include "Utils/SandboxUtils.h"
+
+namespace UE::FileSandboxCore
+{
+void DiscardRemovedPaths(
+	const TSet<FSandboxPath>& InDeletedSandboxPaths, IPlatformFile& InLowerLevel, 
+	TArray<FName>& OutPackagesPendingPurge,
+	TFunctionRef<void(const FFileChange&)> InProcessFileChangeCallback
+	)
+{
+	for (const FSandboxPath& NonSandboxPath : InDeletedSandboxPaths)
+	{
+		// If this file maps to a package then we need to flush its linker so that we can remove the file from the sandbox
+		FName PackageName;
+		FlushPackageFile(NonSandboxPath, &PackageName);
+
+		if (InLowerLevel.FileExists(*NonSandboxPath))
+		{
+			if (!PackageName.IsNone())
+			{
+				OutPackagesPendingPurge.Add(PackageName);
+			}
+			
+			InProcessFileChangeCallback(FFileChange{ NonSandboxPath, EFileChangeAction::Added });
+		}
+	}
+}
+
+void DiscardAddedAndModifiedPaths(
+	TConstArrayView<FSandboxMountPoint> InMountPoints, IPlatformFile& InLowerLevel, 
+	TArray<FName>& OutPackagesPendingHotReload,
+	TArray<FName>& OutPackagesPendingPurge,
+	TFunctionRef<void(const FFileChange&)> InProcessFileChangeCallback
+	)
+{
+	for (const FSandboxMountPoint& SandboxMountPoint : InMountPoints)
+	{
+		// Modify any files that exist in both the non-sandbox and sandbox directories
+		// Delete any files that exist in the sandbox but don't exist in the non-sandbox directory
+		InLowerLevel.IterateDirectoryRecursively(*SandboxMountPoint.Path.GetSandboxPath(), 
+			[&InLowerLevel, &SandboxMountPoint, &OutPackagesPendingHotReload, &OutPackagesPendingPurge, &InProcessFileChangeCallback]
+			(const TCHAR* InFilenameOrDirectory, const bool InIsDirectory) -> bool
+		{
+			if (!InIsDirectory)
+			{
+				const FSandboxedPlatformFilePath RemappedFilePath = FSandboxedPlatformFilePath::CreateNonSandboxPath(FPaths::ConvertRelativePathToFull(InFilenameOrDirectory), SandboxMountPoint.Path);
+
+				// If this file maps to a package then we need to flush its linker so that we can remove the file from the sandbox
+				FName PackageName;
+				FlushPackageFile(RemappedFilePath.GetNonSandboxPath(), &PackageName, false);
+
+				if (InLowerLevel.FileExists(*RemappedFilePath.GetNonSandboxPath()))
+				{
+					if (!PackageName.IsNone())
+					{
+						OutPackagesPendingHotReload.Add(PackageName);
+						OutPackagesPendingPurge.Remove(PackageName);
+					}
+					
+					InProcessFileChangeCallback(FFileChange{ RemappedFilePath.GetNonSandboxPath(), EFileChangeAction::Modified });
+				}
+				else
+				{
+					if (!PackageName.IsNone())
+					{
+						OutPackagesPendingPurge.Add(PackageName);
+						OutPackagesPendingHotReload.Remove(PackageName);
+					}
+					
+					InProcessFileChangeCallback(FFileChange{ RemappedFilePath.GetNonSandboxPath(), EFileChangeAction::Removed });
+				}
+			}
+			return true; // Continue iteration
+		});
+
+		//FString SandboxRootPath = GetSandboxRootPath();
+		//FString Stashed = SandboxRootPath.Replace(TEXT("Sandbox"), TEXT("Stashed"));
+		//LowerLevel->MoveFile(SandboxRootPath, Stashed);
+		// Delete everything under the mount point
+		//LowerLevel->IterateDirectory(*SandboxMountPoint.Path.GetSandboxPath(), [this](const TCHAR* InFilenameOrDirectory, const bool InIsDirectory) -> bool
+		//{
+		//	if (InIsDirectory)
+		//	{
+		//		LowerLevel->DeleteDirectoryRecursively(InFilenameOrDirectory);
+		//	}
+		//	else
+		//	{
+		//		LowerLevel->DeleteFile(InFilenameOrDirectory);
+		//	}
+		//	return true; // Continue iteration
+		//});
+	}
+}
+}
diff --git a/Engine/Plugins/Developer/Sandbox/FileSandbox/Source/FileSandboxCore/Private/Utils/Persist/DiscardUtils.h b/Engine/Plugins/Developer/Sandbox/FileSandbox/Source/FileSandboxCore/Private/Utils/Persist/DiscardUtils.h
new file mode 100644
--- /dev/null
+++ b/Engine/Plugins/Developer/Sandbox/FileSandbox/Source/FileSandboxCore/Private/Utils/Persist/DiscardUtils.h
@@ -0,0 +1,28 @@
+﻿// Copyright Epic Games, Inc. All Rights Reserved.
+
+#pragma once
+
+#include "Containers/ContainersFwd.h"
+#include "Sandbox/Platform/SandboxedPlatformFilePath.h"
+#include "Templates/Function.h"
+
+class IPlatformFile;
+
+namespace UE::FileSandboxCore
+{
+struct FFileChange;
+struct FSandboxMountPoint;
+	
+void DiscardRemovedPaths(
+	const TSet<FSandboxPath>& InDeletedSandboxPaths, IPlatformFile& InLowerLevel, 
+	TArray<FName>& OutPackagesPendingPurge,
+	TFunctionRef<void(const FFileChange&)> InProcessFileChangeCallback = [](const FFileChange&){}
+	);
+	
+void DiscardAddedAndModifiedPaths(
+	TConstArrayView<FSandboxMountPoint> InMountPoints, IPlatformFile& InLowerLevel, 
+	TArray<FName>& OutPackagesPendingHotReload,
+	TArray<FName>& OutPackagesPendingPurge,
+	TFunctionRef<void(const FFileChange&)> InProcessFileChangeCallback = [](const FFileChange&){}
+	);
+}
diff --git a/Engine/Plugins/Developer/Sandbox/FileSandbox/Source/FileSandboxCore/Private/Utils/SandboxUtils.cpp b/Engine/Plugins/Developer/Sandbox/FileSandbox/Source/FileSandboxCore/Private/Utils/SandboxUtils.cpp
--- a/Engine/Plugins/Developer/Sandbox/FileSandbox/Source/FileSandboxCore/Private/Utils/SandboxUtils.cpp
+++ b/Engine/Plugins/Developer/Sandbox/FileSandbox/Source/FileSandboxCore/Private/Utils/SandboxUtils.cpp
@@ -2,11 +2,7 @@
 
 #include "SandboxUtils.h"
 
-#include "DirectoryWatcherModule.h"
-#include "IDirectoryWatcher.h"
-#include "Elements/Columns/TypedElementRevisionControlColumns.h"
 #include "HAL/FileManager.h"
-#include "HAL/IConsoleManager.h"
 #include "Modules/ModuleManager.h"
 #include "UObject/Linker.h"
 #include "UObject/LinkerLoad.h"
@@ -14,6 +10,11 @@
 #include "UObject/UObjectGlobals.h"
 #include "Utils/PackageSandboxUtils.h"
 
+#if UE_SANDBOX_WITH_DIRECTORY_WATCHER
+#include "DirectoryWatcherModule.h"
+#include "IDirectoryWatcher.h"
+#endif
+
 namespace UE::FileSandboxCore
 {
 void ReplayChanges(TConstArrayView<FString> DeletedFiles, TConstArrayView<FString> ModifiedFiles, TConstArrayView<FString> AddedFiles)
@@ -45,19 +46,21 @@ void ReplayChanges(TConstArrayView<FString> DeletedFiles, TConstArrayView<FStrin
 	PurgePackages(PackagesPendingPurge);
 
 	// This is needed to make added files show up in the content browser
+#if UE_SANDBOX_WITH_DIRECTORY_WATCHER
 	TArray<FFileChangeData> FileChanges;
 	Algo::Transform(AddedFiles, FileChanges, [](const FString& File){ return FFileChangeData(File, FFileChangeData::FCA_Added); });
 	if (FDirectoryWatcherModule* DirectoryWatcherModule = GetDirectoryWatcherModuleIfLoaded())
 	{
 		DirectoryWatcherModule->RegisterExternalChanges(FileChanges);
 	}
+#endif
 
 	SynchronizeAssetRegistry();
 }
 	
 void SynchronizeAssetRegistry()
 {
-#if WITH_EDITOR
+#if UE_SANDBOX_WITH_DIRECTORY_WATCHER
 	IDirectoryWatcher* DirectoryWatcher = GetDirectoryWatcherIfLoaded();
 	if (!DirectoryWatcher)
 	{
@@ -105,12 +108,15 @@ bool FlushPackageFile(const FString& InFilename, FName* OutPackageName, bool bFo
 	return false;
 }
 
+#if UE_SANDBOX_WITH_DIRECTORY_WATCHER
 FDirectoryWatcherModule& GetDirectoryWatcherModule()
 {
 	static const FName DirectoryWatcherModuleName = TEXT("DirectoryWatcher");
 	return FModuleManager::LoadModuleChecked<FDirectoryWatcherModule>(DirectoryWatcherModuleName);
 }
+#endif
 
+#if UE_SANDBOX_WITH_DIRECTORY_WATCHER
 FDirectoryWatcherModule* GetDirectoryWatcherModuleIfLoaded()
 {
 	static const FName DirectoryWatcherModuleName = TEXT("DirectoryWatcher");
@@ -120,13 +126,17 @@ FDirectoryWatcherModule* GetDirectoryWatcherModuleIfLoaded()
 	}
 	return nullptr;
 }
+#endif
 
+#if UE_SANDBOX_WITH_DIRECTORY_WATCHER
 IDirectoryWatcher* GetDirectoryWatcher()
 {
 	FDirectoryWatcherModule& DirectoryWatcherModule = GetDirectoryWatcherModule();
 	return DirectoryWatcherModule.Get();
 }
+#endif
 
+#if UE_SANDBOX_WITH_DIRECTORY_WATCHER
 IDirectoryWatcher* GetDirectoryWatcherIfLoaded()
 {
 	if (FDirectoryWatcherModule* DirectoryWatcherModule = GetDirectoryWatcherModuleIfLoaded())
@@ -135,4 +145,5 @@ IDirectoryWatcher* GetDirectoryWatcherIfLoaded()
 	}
 	return nullptr;
 }
+#endif
 }
diff --git a/Engine/Plugins/Developer/Sandbox/FileSandbox/Source/FileSandboxCore/Private/Utils/SandboxUtils.h b/Engine/Plugins/Developer/Sandbox/FileSandbox/Source/FileSandboxCore/Private/Utils/SandboxUtils.h
--- a/Engine/Plugins/Developer/Sandbox/FileSandbox/Source/FileSandboxCore/Private/Utils/SandboxUtils.h
+++ b/Engine/Plugins/Developer/Sandbox/FileSandbox/Source/FileSandboxCore/Private/Utils/SandboxUtils.h
@@ -22,8 +22,10 @@ void FlushPackageLoading(const FString& InPackageName, bool bForceBulkDataLoad =
 /** Alternate variant that flushes package loading by filename instead. */
 bool FlushPackageFile(const FString& InFilename, FName* OutPackageName = nullptr, bool bForceLoad = true);
 	
+#if UE_SANDBOX_WITH_DIRECTORY_WATCHER
 FDirectoryWatcherModule& GetDirectoryWatcherModule();
 FDirectoryWatcherModule* GetDirectoryWatcherModuleIfLoaded();
 IDirectoryWatcher* GetDirectoryWatcher();
 IDirectoryWatcher* GetDirectoryWatcherIfLoaded();
+#endif
 }
