File Sandbox: Source control proxy now optionally forces redirectory to be created when assets are renamed when source control is set up.

This feature is required by Concert / Multi-User, and is enabled by default.
In order for the redirector to be created, the user must be connected to source control and rename a file in the sandbox.

#Synced-CL 47637010

diff --git a/Engine/Plugins/Developer/Sandbox/FileSandbox/Source/FileSandboxCore/Private/Sandbox/SandboxInstance.cpp b/Engine/Plugins/Developer/Sandbox/FileSandbox/Source/FileSandboxCore/Private/Sandbox/SandboxInstance.cpp
--- a/Engine/Plugins/Developer/Sandbox/FileSandbox/Source/FileSandboxCore/Private/Sandbox/SandboxInstance.cpp
+++ b/Engine/Plugins/Developer/Sandbox/FileSandbox/Source/FileSandboxCore/Private/Sandbox/SandboxInstance.cpp
@@ -54,7 +54,7 @@ FSandboxInstance::FSandboxInstance(
 	, ManifestContent(MoveTemp(InManifestData))
 	, MetaData(MoveTemp(InMetaData))
 	, PlatformFile(MakeUnique<FSandboxPlatformFile>(RootDirectory))
-	, SourceControlSandbox(*this)
+	, SourceControlSandbox(*this, MakeConfigFromInitArgs(InitArgs))
 {
 	PlatformFile->Initialize(&InitArgs.GetPlatformFile(), TEXT(""));
 	ReplayChanges(ManifestContent.DeletedFiles.GetFiles(), ManifestContent.ModifiedFiles.GetFiles(), ManifestContent.AddedFiles.GetFiles());
diff --git a/Engine/Plugins/Developer/Sandbox/FileSandbox/Source/FileSandboxCore/Private/SourceControl/Proxy/SandboxSourceControlProviderProxy.cpp b/Engine/Plugins/Developer/Sandbox/FileSandbox/Source/FileSandboxCore/Private/SourceControl/Proxy/SandboxSourceControlProviderProxy.cpp
new file mode 100644
--- /dev/null
+++ b/Engine/Plugins/Developer/Sandbox/FileSandbox/Source/FileSandboxCore/Private/SourceControl/Proxy/SandboxSourceControlProviderProxy.cpp
@@ -0,0 +1,143 @@
+﻿// Copyright Epic Games, Inc. All Rights Reserved.
+
+#include "SandboxSourceControlProviderProxy.h"
+
+#include "Features/IModularFeatures.h"
+#include "GenericPlatform/GenericPlatformFile.h"
+#include "HAL/PlatformFileManager.h"
+#include "SandboxSourceControlStateProxy.h"
+
+namespace UE::FileSandboxCore
+{
+namespace Private
+{
+static bool IsProhibitedOp(const FSourceControlOperationRef& InOperation)
+{
+	const FName OpName = InOperation->GetName();
+	// We're not supposed to alter source control state while in sandbox - we touch the files directly and only affect source control when persisting.
+	// Some editor operations issue a revert followed by file op, e.g. deleting an edited asset will first revert the asset, then delete it.
+	return OpName == TEXT("Revert")	
+		// This would show a popup a status update widget, don't show it in sandbox.
+		|| OpName == TEXT("UpdateStatus"); 
+}
+	
+static TOptional<ECommandResult::Type> HandleSandboxedOp(const FSourceControlOperationRef& InOperation, const TArray<FString>& InFiles)
+{
+	const FName OpName = InOperation->GetName();
+	
+	// Edits should not show up in the external source control programming until they are persisted by the sandbox
+	if (OpName == TEXT("CheckOut"))
+	{
+		return ECommandResult::Succeeded;
+	}
+	
+	// Added files should not show up in the external source control programming until they are persisted by the sandbox
+	if (OpName == TEXT("MarkForAdd"))
+	{
+		return ECommandResult::Succeeded;
+	}
+	
+	// The contract ISourceControlProvider handles deleting the file. We'll just reroute it to our sandbox. 
+	// If the changes end up being persisted, the source control changes will be routed to the real source control then.
+	if (OpName == TEXT("Delete"))
+	{
+		IPlatformFile& PlatformFile = FPlatformFileManager::Get().GetPlatformFile();
+		bool bAllSuccessful = true;
+		for (const FString& InFile : InFiles)
+		{
+			bAllSuccessful &= PlatformFile.DeleteFile(*InFile);
+		}
+		return bAllSuccessful ? ECommandResult::Succeeded : ECommandResult::Failed;
+	}
+	
+	return {};
+}
+	
+static TOptional<ECommandResult::Type> InterceptSourceControlOp(
+	const FSourceControlOperationRef& InOperation , const TArray<FString>& InFiles
+	)
+{
+	// Handle forwarding equivalents of edit, add, and remove to sandbox.
+	if (const TOptional<ECommandResult::Type> Result = HandleSandboxedOp(InOperation, InFiles))
+	{
+		return *Result;
+	}
+	
+	if (IsProhibitedOp(InOperation))
+	{
+		return ECommandResult::Succeeded;
+	}
+	
+	return {};
+}
+}
+	
+FSandboxSourceControlProviderProxy::FSandboxSourceControlProviderProxy(
+	FSandboxInstance& InInstance, const FSourceControlProxyConfig& InStateProxyConfig
+	)
+	: Super(&ISourceControlModule::Get().GetProvider())
+	, Instance(InInstance)
+	, StateProxyConfig(InStateProxyConfig)
+{
+	IModularFeatures::Get().RegisterModularFeature("SourceControl", this);
+}
+
+FSandboxSourceControlProviderProxy::~FSandboxSourceControlProviderProxy()
+{
+	IModularFeatures::Get().UnregisterModularFeature("SourceControl", this);
+}
+
+const FName& FSandboxSourceControlProviderProxy::GetName() const
+{
+	static FName Name(TEXT("File Sandbox"));
+	return Name;
+}
+
+ECommandResult::Type FSandboxSourceControlProviderProxy::Execute(
+	const FSourceControlOperationRef& InOperation,
+	FSourceControlChangelistPtr InChangelist, const TArray<FString>& InFiles, EConcurrency::Type InConcurrency,
+	const FSourceControlOperationComplete& InOperationCompleteDelegate
+	)
+{
+	if (const TOptional<ECommandResult::Type> Result = Private::InterceptSourceControlOp(InOperation, InFiles))
+	{
+		InOperationCompleteDelegate.ExecuteIfBound(InOperation, *Result);
+		return *Result;
+	}
+	
+	return Super::Execute(InOperation, InChangelist, InFiles, InConcurrency, InOperationCompleteDelegate);
+}
+
+ECommandResult::Type FSandboxSourceControlProviderProxy::GetState(
+	const TArray<FString>& InFiles, TArray<FSourceControlStateRef>& OutState, EStateCacheUsage::Type InStateCacheUsage
+	)
+{
+	ECommandResult::Type Result = ECommandResult::Failed;
+
+	if (UnderlyingProvider)
+	{
+		Result = UnderlyingProvider->GetState(InFiles, OutState, InStateCacheUsage);
+	}
+
+	if (Result == ECommandResult::Failed)
+	{
+		// Even when the provider fails, return a dummy proxy so we override some of the behavior, such as the IsLocal override for redirectors.
+		Result = ECommandResult::Succeeded;
+
+		OutState.Reset();
+		for (const FString& File : InFiles)
+		{
+			OutState.Add(MakeShared<FSandboxSourceControlStateProxy, ESPMode::ThreadSafe>(File, StateProxyConfig));
+		}
+	}
+	else
+	{
+		for (FSourceControlStateRef& State : OutState)
+		{
+			State = MakeShared<FSandboxSourceControlStateProxy, ESPMode::ThreadSafe>(State, StateProxyConfig);
+		}
+	}
+
+	return Result;
+}
+}
diff --git a/Engine/Plugins/Developer/Sandbox/FileSandbox/Source/FileSandboxCore/Private/SourceControl/Proxy/SandboxSourceControlProviderProxy.h b/Engine/Plugins/Developer/Sandbox/FileSandbox/Source/FileSandboxCore/Private/SourceControl/Proxy/SandboxSourceControlProviderProxy.h
new file mode 100644
--- /dev/null
+++ b/Engine/Plugins/Developer/Sandbox/FileSandbox/Source/FileSandboxCore/Private/SourceControl/Proxy/SandboxSourceControlProviderProxy.h
@@ -0,0 +1,58 @@
+﻿// Copyright Epic Games, Inc. All Rights Reserved.
+
+#pragma once
+
+#include "HAL/Platform.h"
+#include "SandboxSourceControlStateProxy.h"
+#include "SourceControlProviderProxyBase.h"
+#include "Templates/UnrealTemplate.h"
+
+namespace UE::FileSandboxCore
+{
+class FSandboxInstance;
+struct FSourceControlProxyConfig;
+
+/** 
+ * Intercepts any file changes the engine attempts to make and routes them FSandboxInstance.
+ * 
+ * When set-up with source control, some tools in the engine route I/O operations to the external source control tool instead of to IPlatformFile.
+ * The point of this is that the external program handles the I/O operations, e.g. "p4 delete" deletes the file and registers it with a changelist.
+ * 
+ * However, during sandbox we're not supposed to make any external changes - so we'll "swallow" those changes and re-reroute to FSandboxInstance instead.
+ * FSandboxInstance will handle re-issuing the source control commands again if the changes are persisted.
+ */
+class FSandboxSourceControlProviderProxy 
+	: public FBaseSourceControlProviderProxy
+	, public FNoncopyable
+{
+	using Super = FBaseSourceControlProviderProxy;
+public:
+	
+	explicit FSandboxSourceControlProviderProxy(FSandboxInstance& InInstance UE_LIFETIMEBOUND, const FSourceControlProxyConfig& InStateProxyConfig);
+	~FSandboxSourceControlProviderProxy();
+	
+	void SetUnderlyingSourceControlProvider(ISourceControlProvider* InProvider) { UnderlyingProvider = InProvider; }
+	
+	//~ Begin ISourceControlProvider Interface
+	virtual const FName& GetName() const override;
+	virtual ECommandResult::Type Execute(
+		const FSourceControlOperationRef& InOperation, FSourceControlChangelistPtr InChangelist, const TArray<FString>& InFiles, 
+		EConcurrency::Type InConcurrency = EConcurrency::Synchronous, 
+		const FSourceControlOperationComplete& InOperationCompleteDelegate = FSourceControlOperationComplete() 
+		) override;
+	virtual ECommandResult::Type GetState(
+		const TArray<FString>& InFiles, TArray<FSourceControlStateRef>& OutState, EStateCacheUsage::Type InStateCacheUsage
+		) override;
+	//~ End ISourceControlProvider Interface
+	
+private:
+	
+	/** Instance of the sandbox for which we've been created. */
+	FSandboxInstance& Instance;
+	
+	/** Config to construct the proxy state instances with. */
+	const FSourceControlProxyConfig StateProxyConfig;
+};
+}
+
+
diff --git a/Engine/Plugins/Developer/Sandbox/FileSandbox/Source/FileSandboxCore/Private/SourceControl/Proxy/SandboxSourceControlStateProxy.cpp b/Engine/Plugins/Developer/Sandbox/FileSandbox/Source/FileSandboxCore/Private/SourceControl/Proxy/SandboxSourceControlStateProxy.cpp
new file mode 100644
--- /dev/null
+++ b/Engine/Plugins/Developer/Sandbox/FileSandbox/Source/FileSandboxCore/Private/SourceControl/Proxy/SandboxSourceControlStateProxy.cpp
@@ -0,0 +1,31 @@
+﻿// Copyright Epic Games, Inc. All Rights Reserved.
+
+#include "SandboxSourceControlStateProxy.h"
+
+#include "HAL/IConsoleManager.h"
+
+namespace UE::FileSandboxCore
+{
+enum EIsLocalOverrideMode
+{
+	NoOverride,
+	ForceLocal,
+	ForceNonLocal
+};
+static TAutoConsoleVariable<int32> CVarOverrideIsLocal(
+	TEXT("Sandbox.IsLocalOverrideMode"), 
+	EIsLocalOverrideMode::NoOverride, 
+	TEXT("Debug util for ISourceControlState::IsLocal while in sandbox.\n0 - No override (return what SC says),\n1 - Force IsLocal() == true,\n2 - Force IsLocal() == false")
+	);
+	
+bool FSandboxSourceControlStateProxy::IsLocal() const
+{
+	switch (CVarOverrideIsLocal.GetValueOnAnyThread())
+	{
+	case EIsLocalOverrideMode::ForceLocal: return true;
+	case EIsLocalOverrideMode::ForceNonLocal: return false;
+	case EIsLocalOverrideMode::NoOverride: [[fallthrough]];
+	default: return Config.IsLocalOverride ? *Config.IsLocalOverride : Super::IsLocal();
+	}
+}
+}
diff --git a/Engine/Plugins/Developer/Sandbox/FileSandbox/Source/FileSandboxCore/Private/SourceControl/Proxy/SandboxSourceControlStateProxy.h b/Engine/Plugins/Developer/Sandbox/FileSandbox/Source/FileSandboxCore/Private/SourceControl/Proxy/SandboxSourceControlStateProxy.h
new file mode 100644
--- /dev/null
+++ b/Engine/Plugins/Developer/Sandbox/FileSandbox/Source/FileSandboxCore/Private/SourceControl/Proxy/SandboxSourceControlStateProxy.h
@@ -0,0 +1,37 @@
+﻿// Copyright Epic Games, Inc. All Rights Reserved.
+
+#pragma once
+
+#include "SourceControlProxyConfig.h"
+#include "SourceControlStateProxyBase.h"
+
+namespace UE::FileSandboxCore
+{
+/** 
+ * The purpose of this proxy state is to force files to optionally be treated as non-local. 
+ * This achieves that when files are renamed, redirectors are created as well. Some features, such as Multi-User / Concert, require redirectors.
+ */
+class FSandboxSourceControlStateProxy : public FSourceControlStateProxyBase
+{
+	using Super = FSourceControlStateProxyBase;
+public:
+	
+	explicit FSandboxSourceControlStateProxy(FSourceControlStatePtr InProxyState, const FSourceControlProxyConfig& InConfig) 
+		: FSourceControlStateProxyBase(MoveTemp(InProxyState)) 
+		, Config(InConfig)
+	{}
+	explicit FSandboxSourceControlStateProxy(FString InFilename, const FSourceControlProxyConfig& InConfig) 
+		: FSourceControlStateProxyBase(MoveTemp(InFilename)) 
+		, Config(InConfig)
+		{}
+
+	//~ Begin ISourceControlState Interface
+	virtual bool IsLocal() const override;
+	//~ End ISourceControlState Interface
+	
+private:
+	
+	/** Config for our behaviour should override the proxied state. */
+	const FSourceControlProxyConfig Config;
+};
+}
diff --git a/Engine/Plugins/Developer/Sandbox/FileSandbox/Source/FileSandboxCore/Private/SourceControl/Proxy/SourceControlProviderProxyBase.h b/Engine/Plugins/Developer/Sandbox/FileSandbox/Source/FileSandboxCore/Private/SourceControl/Proxy/SourceControlProviderProxyBase.h
new file mode 100644
--- /dev/null
+++ b/Engine/Plugins/Developer/Sandbox/FileSandbox/Source/FileSandboxCore/Private/SourceControl/Proxy/SourceControlProviderProxyBase.h
@@ -0,0 +1,173 @@
+﻿// Copyright Epic Games, Inc. All Rights Reserved.
+
+#pragma once
+
+#include "ISourceControlModule.h"
+#include "ISourceControlProvider.h"
+#if SOURCE_CONTROL_WITH_SLATE
+#include "Widgets/SNullWidget.h"
+#endif
+
+#define WRAP(FuncName, ...) if (ISourceControlProvider* Provider = GetUnderlyingProvider()) { Provider->FuncName(__VA_ARGS__); }
+#define WRAP_RET(Default, FuncName, ...) if (ISourceControlProvider* Provider = GetUnderlyingProvider()) { return Provider->FuncName(__VA_ARGS__); } return Default; 
+
+namespace UE::FileSandboxCore
+{
+/** 
+ * Forwards all calls to an optional, proxied source control provider. 
+ * Acts as base class to reduce boilerplate when proxying source control classes.
+ */
+class FBaseSourceControlProviderProxy : public ISourceControlProvider
+{
+public:
+	
+	explicit FBaseSourceControlProviderProxy(ISourceControlProvider* InActualProvider)
+		: UnderlyingProvider(InActualProvider)
+	{}
+	
+	//~ Begin ISourceControlProvider Interface
+	virtual void Init(bool bForceConnection = true) override
+	{
+		WRAP(Init, bForceConnection)
+	}
+	virtual void Close() override
+	{
+		WRAP(Close);
+	}
+	virtual FText GetStatusText() const override
+	{
+		WRAP_RET(FText::GetEmpty(), GetStatusText)
+	}
+	virtual TMap<EStatus, FString> GetStatus() const override
+	{
+		WRAP_RET({}, GetStatus)
+	}
+	virtual bool IsEnabled() const override
+	{
+		WRAP_RET(true, IsEnabled)
+	}
+	virtual bool IsAvailable() const override
+	{
+		WRAP_RET(true, IsAvailable)
+	}
+	virtual bool QueryStateBranchConfig(const FString& ConfigSrc, const FString& ConfigDest) override
+	{
+		WRAP_RET(false, QueryStateBranchConfig, ConfigSrc, ConfigDest)
+	}
+	virtual void RegisterStateBranches(const TArray<FString>& BranchNames, const FString& ContentRoot) override
+	{
+		WRAP(RegisterStateBranches, BranchNames, ContentRoot)
+	}
+	virtual int32 GetStateBranchIndex(const FString& BranchName) const override
+	{
+		WRAP_RET(INDEX_NONE, GetStateBranchIndex, BranchName);
+	}
+	virtual bool GetStateBranchAtIndex(int32 BranchIndex, FString& OutBranchName) const override
+	{
+		WRAP_RET(false, GetStateBranchAtIndex, BranchIndex, OutBranchName);
+	}
+	virtual ECommandResult::Type GetState(const TArray<FString>& InFiles, TArray<FSourceControlStateRef>& OutState, EStateCacheUsage::Type InStateCacheUsage) override
+	{
+		WRAP_RET(ECommandResult::Failed, GetState, InFiles, OutState, InStateCacheUsage)
+	}
+	virtual ECommandResult::Type GetState(const TArray<FSourceControlChangelistRef>& InChangelists, TArray<FSourceControlChangelistStateRef>& OutState, EStateCacheUsage::Type InStateCacheUsage) override
+	{
+		WRAP_RET(ECommandResult::Failed, GetState, InChangelists, OutState, InStateCacheUsage)
+	}
+	virtual TArray<FSourceControlStateRef> GetCachedStateByPredicate(TFunctionRef<bool(const FSourceControlStateRef&)> Predicate) const override
+	{
+		WRAP_RET({}, GetCachedStateByPredicate, Predicate)
+	}
+	virtual FDelegateHandle RegisterSourceControlStateChanged_Handle(const FSourceControlStateChanged::FDelegate& SourceControlStateChanged) override
+	{
+		WRAP_RET({}, RegisterSourceControlStateChanged_Handle, SourceControlStateChanged)
+	}
+	virtual void UnregisterSourceControlStateChanged_Handle(FDelegateHandle Handle) override
+	{
+		WRAP(UnregisterSourceControlStateChanged_Handle, Handle)
+	}
+	virtual ECommandResult::Type Execute(const FSourceControlOperationRef& InOperation, FSourceControlChangelistPtr InChangelist, const TArray<FString>& InFiles, EConcurrency::Type InConcurrency = EConcurrency::Synchronous, const FSourceControlOperationComplete& InOperationCompleteDelegate = FSourceControlOperationComplete()) override
+	{
+		WRAP_RET(ECommandResult::Failed, Execute, InOperation, InChangelist, InFiles, InConcurrency, InOperationCompleteDelegate)
+	} 
+	virtual bool CanExecuteOperation(const FSourceControlOperationRef& InOperation) const override
+	{
+		WRAP_RET(false, CanExecuteOperation, InOperation)
+	}
+	virtual bool CanCancelOperation(const FSourceControlOperationRef& InOperation) const override
+	{
+		WRAP_RET(false, CanCancelOperation, InOperation)
+	}
+	virtual void CancelOperation(const FSourceControlOperationRef& InOperation) override
+	{
+		WRAP(CancelOperation, InOperation)
+	}
+	virtual TArray<TSharedRef<class ISourceControlLabel>> GetLabels(const FString& InMatchingSpec) const override
+	{
+		WRAP_RET({}, GetLabels, InMatchingSpec)
+	}
+	virtual TArray<FSourceControlChangelistRef> GetChangelists(EStateCacheUsage::Type InStateCacheUsage) override
+	{
+		WRAP_RET({}, GetChangelists, InStateCacheUsage)
+	}
+	virtual bool UsesLocalReadOnlyState() const override
+	{
+		WRAP_RET(false, UsesLocalReadOnlyState)
+	}
+	virtual bool UsesChangelists() const override
+	{
+		WRAP_RET(false, UsesChangelists)
+	}
+	virtual bool UsesUncontrolledChangelists() const override
+	{
+		WRAP_RET(false, UsesUncontrolledChangelists)
+	}
+	virtual bool UsesCheckout() const override
+	{
+		WRAP_RET(false, UsesCheckout)
+	}
+	virtual bool UsesFileRevisions() const override
+	{
+		WRAP_RET(false, UsesFileRevisions)
+	}
+	virtual bool UsesSnapshots() const override
+	{
+		WRAP_RET(false, UsesSnapshots)
+	}
+	virtual bool AllowsDiffAgainstDepot() const override
+	{
+		WRAP_RET(false, AllowsDiffAgainstDepot)
+	}
+	virtual TOptional<bool> IsAtLatestRevision() const override
+	{
+		WRAP_RET({}, IsAtLatestRevision)
+	}
+	virtual TOptional<int> GetNumLocalChanges() const override
+	{
+		WRAP_RET({}, GetNumLocalChanges)
+	}
+	virtual void Tick() override
+	{
+		WRAP(Tick)
+	}
+#if SOURCE_CONTROL_WITH_SLATE
+	virtual TSharedRef<class SWidget> MakeSettingsWidget() const override
+	{
+		WRAP_RET(SNullWidget::NullWidget, MakeSettingsWidget)
+	}
+#endif
+	//~ End ISourceControlProvider Interface
+	
+protected:
+	
+	ISourceControlProvider* UnderlyingProvider;
+	
+	ISourceControlProvider* GetUnderlyingProvider() const 
+	{
+		return UnderlyingProvider;
+	}
+};
+}
+
+#undef WRAP
+#undef WRAP_RET
\ No newline at end of file
diff --git a/Engine/Plugins/Developer/Sandbox/FileSandbox/Source/FileSandboxCore/Private/SourceControl/Proxy/SourceControlProxyConfig.h b/Engine/Plugins/Developer/Sandbox/FileSandbox/Source/FileSandboxCore/Private/SourceControl/Proxy/SourceControlProxyConfig.h
new file mode 100644
--- /dev/null
+++ b/Engine/Plugins/Developer/Sandbox/FileSandbox/Source/FileSandboxCore/Private/SourceControl/Proxy/SourceControlProxyConfig.h
@@ -0,0 +1,30 @@
+﻿// Copyright Epic Games, Inc. All Rights Reserved.
+
+#pragma once
+
+#include "Misc/Optional.h"
+#include "Types/Manager/SandboxInitArgs.h"
+
+namespace UE::FileSandboxCore
+{
+struct FSourceControlProxyConfig
+{
+	/** If set, return this value from IsLocal. */
+	TOptional<bool> IsLocalOverride;
+};
+	
+/** @return Source control proxy config constructed from the sandbox init args. */
+FSourceControlProxyConfig MakeConfigFromInitArgs(const FSandboxInitArgs& InInitArgs);
+}
+
+namespace UE::FileSandboxCore
+{
+inline FSourceControlProxyConfig MakeConfigFromInitArgs(const FSandboxInitArgs& InInitArgs)
+{
+	const bool bForceNonLocal = EnumHasAnyFlags(InInitArgs.Flags, ESandboxInitFlags::ForceRedirectorsOnRenameInSourceControl);
+	return FSourceControlProxyConfig
+	{ 
+		.IsLocalOverride = bForceNonLocal ? false : TOptional<bool>{}
+	};
+}
+}
diff --git a/Engine/Plugins/Developer/Sandbox/FileSandbox/Source/FileSandboxCore/Private/SourceControl/Proxy/SourceControlStateProxyBase.h b/Engine/Plugins/Developer/Sandbox/FileSandbox/Source/FileSandboxCore/Private/SourceControl/Proxy/SourceControlStateProxyBase.h
new file mode 100644
--- /dev/null
+++ b/Engine/Plugins/Developer/Sandbox/FileSandbox/Source/FileSandboxCore/Private/SourceControl/Proxy/SourceControlStateProxyBase.h
@@ -0,0 +1,171 @@
+﻿// Copyright Epic Games, Inc. All Rights Reserved.
+
+#pragma once
+
+#include "ISourceControlState.h"
+#include "HAL/FileManager.h"
+
+namespace UE::FileSandboxCore
+{
+/** 
+ * Forwards all calls to an optional, proxied source control state. 
+ * Acts as base class to reduce boilerplate when proxying source control classes.
+ */
+class FSourceControlStateProxyBase : public ISourceControlState
+{
+public:
+	
+	explicit FSourceControlStateProxyBase(FSourceControlStatePtr InProxyState) : ActualState(MoveTemp(InProxyState)) {}
+	explicit FSourceControlStateProxyBase(FString InFilename) 
+		: ActualState(nullptr)
+		, CachedFilename(MoveTemp(InFilename))
+		, CachedTimestamp(IFileManager::Get().GetTimeStamp(*CachedFilename))
+	{}
+	
+	//~ Begin ISourceControlState Interface
+	virtual int32 GetHistorySize() const override { return ActualState.IsValid() ? ActualState->GetHistorySize() : 0; }
+	virtual TSharedPtr<ISourceControlRevision> GetHistoryItem(int32 HistoryIndex) const override 
+	{ 
+		return ActualState.IsValid() ? ActualState->GetHistoryItem(HistoryIndex) : TSharedPtr<ISourceControlRevision>();
+	}
+	virtual TSharedPtr<ISourceControlRevision> FindHistoryRevision(int32 RevisionNumber) const override
+	{
+		return ActualState.IsValid() ? ActualState->FindHistoryRevision(RevisionNumber) : TSharedPtr<ISourceControlRevision>();
+	}
+	virtual TSharedPtr<ISourceControlRevision> FindHistoryRevision(const FString& InRevision) const override
+	{
+		return ActualState.IsValid() ? ActualState->FindHistoryRevision(InRevision) : TSharedPtr<ISourceControlRevision>();
+	}
+	virtual TSharedPtr<ISourceControlRevision> GetCurrentRevision() const override
+	{
+		return ActualState.IsValid() ? ActualState->GetCurrentRevision() : TSharedPtr<ISourceControlRevision>();
+	}
+	virtual FResolveInfo GetResolveInfo() const override 
+	{ 
+		return ActualState.IsValid() ? ActualState->GetResolveInfo() : FResolveInfo();
+	}
+	virtual FSlateIcon GetIcon() const override
+	{
+		return ActualState.IsValid() ? ActualState->GetIcon() : FSlateIcon();
+	}
+	virtual FText GetDisplayName() const override
+	{
+		return ActualState.IsValid() ? ActualState->GetDisplayName() : FText::GetEmpty();
+	}
+	virtual FText GetDisplayTooltip() const override
+	{
+		return ActualState.IsValid() ? ActualState->GetDisplayTooltip() : FText::GetEmpty();
+	}
+	virtual const FString& GetFilename() const override
+	{
+		return ActualState.IsValid() ? ActualState->GetFilename() : CachedFilename;
+	}
+	virtual const FDateTime& GetTimeStamp() const override
+	{
+		return ActualState.IsValid() ? ActualState->GetTimeStamp() : CachedTimestamp;
+	}
+	virtual bool CanCheckIn() const override
+	{
+		return ActualState.IsValid() && ActualState->CanCheckIn();
+	}
+	virtual bool CanCheckout() const override
+	{
+		return ActualState.IsValid() && ActualState->CanCheckout();
+	}
+	virtual bool IsCheckedOut() const override
+	{
+		return ActualState.IsValid() && ActualState->IsCheckedOut();
+	}
+	virtual bool IsCheckedOutOther(FString* Who = 0) const override
+	{
+		return ActualState.IsValid() && ActualState->IsCheckedOutOther(Who);
+	}
+	virtual bool IsCheckedOutInOtherBranch(const FString& CurrentBranch = FString()) const override
+	{
+		return ActualState.IsValid() && ActualState->IsCheckedOutInOtherBranch(CurrentBranch);
+	}
+	virtual bool IsModifiedInOtherBranch(const FString& CurrentBranch = FString()) const override
+	{
+		return ActualState.IsValid() && ActualState->IsModifiedInOtherBranch(CurrentBranch);
+	}
+	virtual bool IsCheckedOutOrModifiedInOtherBranch(const FString& CurrentBranch = FString()) const override
+	{
+		return ActualState.IsValid() && ActualState->IsCheckedOutOrModifiedInOtherBranch(CurrentBranch);
+	}
+	virtual TArray<FString> GetCheckedOutBranches() const override
+	{
+		return ActualState.IsValid() ? ActualState->GetCheckedOutBranches() : TArray<FString>();
+	}
+	virtual FString GetOtherUserBranchCheckedOuts() const override
+	{
+		return ActualState.IsValid() ? ActualState->GetOtherUserBranchCheckedOuts() : FString();
+	}
+	virtual bool GetOtherBranchHeadModification(FString& HeadBranchOut, FString& ActionOut, int32& HeadChangeListOut) const override
+	{
+		return ActualState.IsValid() && ActualState->GetOtherBranchHeadModification(HeadBranchOut, ActionOut, HeadChangeListOut);
+	}
+	virtual bool IsCurrent() const override
+	{
+		return ActualState.IsValid() && ActualState->IsCurrent();
+	}
+	virtual bool IsSourceControlled() const override
+	{
+		return ActualState.IsValid() && ActualState->IsSourceControlled();
+	}
+	virtual bool IsLocal() const override
+	{
+		return ActualState.IsValid() && ActualState->IsLocal();
+	}
+	virtual bool IsAdded() const override
+	{
+		return ActualState.IsValid() && ActualState->IsAdded();
+	}
+	virtual bool IsDeleted() const override
+	{
+		return ActualState.IsValid() && ActualState->IsDeleted();
+	}
+	virtual bool IsIgnored() const override
+	{
+		return ActualState.IsValid() && ActualState->IsIgnored();
+	}
+	virtual bool CanEdit() const override
+	{
+		return ActualState.IsValid() && ActualState->CanEdit();
+	}
+	virtual bool CanDelete() const override
+	{
+		return ActualState.IsValid() && ActualState->CanDelete();
+	}
+	virtual bool IsUnknown() const override
+	{
+		return ActualState.IsValid() && ActualState->IsUnknown();
+	}
+	virtual bool IsModified() const override
+	{
+		return ActualState.IsValid() && ActualState->IsModified();
+	}
+	virtual bool CanAdd() const override
+	{
+		return ActualState.IsValid() && ActualState->CanAdd();
+	}
+	virtual bool IsConflicted() const override
+	{
+		return ActualState.IsValid() && ActualState->IsConflicted();
+	}
+	virtual bool CanRevert() const override
+	{
+		return ActualState.IsValid() && ActualState->CanRevert();
+	}
+	//~ End ISourceControlState Interface
+	
+protected:
+	
+	/** The underlying state we proxy through. */
+	FSourceControlStatePtr ActualState;
+	
+	/** The name of the file we represent (only when ActualState is null) */
+	FString CachedFilename;
+	/** The timestamp of the file we represent (only when ActualState is null) */
+	FDateTime CachedTimestamp;
+};
+}
diff --git a/Engine/Plugins/Developer/Sandbox/FileSandbox/Source/FileSandboxCore/Private/SourceControl/SandboxSourceControlProviderProxy.cpp b/Engine/Plugins/Developer/Sandbox/FileSandbox/Source/FileSandboxCore/Private/SourceControl/SandboxSourceControlProviderProxy.cpp
deleted file mode 100644
diff --git a/Engine/Plugins/Developer/Sandbox/FileSandbox/Source/FileSandboxCore/Private/SourceControl/SandboxSourceControlProviderProxy.h b/Engine/Plugins/Developer/Sandbox/FileSandbox/Source/FileSandboxCore/Private/SourceControl/SandboxSourceControlProviderProxy.h
deleted file mode 100644
diff --git a/Engine/Plugins/Developer/Sandbox/FileSandbox/Source/FileSandboxCore/Private/SourceControl/SandboxedSourceControl.cpp b/Engine/Plugins/Developer/Sandbox/FileSandbox/Source/FileSandboxCore/Private/SourceControl/SandboxedSourceControl.cpp
--- a/Engine/Plugins/Developer/Sandbox/FileSandbox/Source/FileSandboxCore/Private/SourceControl/SandboxedSourceControl.cpp
+++ b/Engine/Plugins/Developer/Sandbox/FileSandbox/Source/FileSandboxCore/Private/SourceControl/SandboxedSourceControl.cpp
@@ -4,9 +4,9 @@
 
 namespace UE::FileSandboxCore
 {
-FSandboxedSourceControl::FSandboxedSourceControl(FSandboxInstance& InSandboxInstance)
+FSandboxedSourceControl::FSandboxedSourceControl(FSandboxInstance& InSandboxInstance, const FSourceControlProxyConfig& InStateProxyConfig)
 	: SandboxInstance(InSandboxInstance)
-	, SourceControlProviderProxy(InSandboxInstance)
+	, SourceControlProviderProxy(InSandboxInstance, InStateProxyConfig)
 	, Override(SourceControlProviderProxy)
 {
 	Override.OnProxiedProviderChanged().AddLambda([this](ISourceControlProvider& InRestoreProvider)
diff --git a/Engine/Plugins/Developer/Sandbox/FileSandbox/Source/FileSandboxCore/Private/SourceControl/SandboxedSourceControl.h b/Engine/Plugins/Developer/Sandbox/FileSandbox/Source/FileSandboxCore/Private/SourceControl/SandboxedSourceControl.h
--- a/Engine/Plugins/Developer/Sandbox/FileSandbox/Source/FileSandboxCore/Private/SourceControl/SandboxedSourceControl.h
+++ b/Engine/Plugins/Developer/Sandbox/FileSandbox/Source/FileSandboxCore/Private/SourceControl/SandboxedSourceControl.h
@@ -3,20 +3,21 @@
 #pragma once
 
 #include "HAL/Platform.h"
-#include "SandboxSourceControlProviderProxy.h"
+#include "Proxy/SandboxSourceControlProviderProxy.h"
 #include "ScopedSourceControlOverride.h"
 #include "Templates/UnrealTemplate.h"
 
 namespace UE::FileSandboxCore
 {
 class FSandboxInstance;
+struct FSourceControlProxyConfig;
 	
 /** Interface for FSandboxInstance to interact with source control. */
 class FSandboxedSourceControl : public FNoncopyable
 {
 public:
 
-	explicit FSandboxedSourceControl(FSandboxInstance& InSandboxInstance UE_LIFETIMEBOUND);
+	explicit FSandboxedSourceControl(FSandboxInstance& InSandboxInstance UE_LIFETIMEBOUND, const FSourceControlProxyConfig& InStateProxyConfig);
 	
 	/** @return The real source control provider that we're overriding. */
 	ISourceControlProvider& GetProxiedSourceControl() const { return Override.GetProxiedSourceControlProvider(); }
diff --git a/Engine/Plugins/Developer/Sandbox/FileSandbox/Source/FileSandboxCore/Private/SourceControl/SourceControlProviderProxyBase.h b/Engine/Plugins/Developer/Sandbox/FileSandbox/Source/FileSandboxCore/Private/SourceControl/SourceControlProviderProxyBase.h
deleted file mode 100644
diff --git a/Engine/Plugins/Developer/Sandbox/FileSandbox/Source/FileSandboxCore/Public/Types/Manager/SandboxInitArgs.h b/Engine/Plugins/Developer/Sandbox/FileSandbox/Source/FileSandboxCore/Public/Types/Manager/SandboxInitArgs.h
--- a/Engine/Plugins/Developer/Sandbox/FileSandbox/Source/FileSandboxCore/Public/Types/Manager/SandboxInitArgs.h
+++ b/Engine/Plugins/Developer/Sandbox/FileSandbox/Source/FileSandboxCore/Public/Types/Manager/SandboxInitArgs.h
@@ -3,7 +3,7 @@
 #pragma once
 
 #include "HAL/PlatformFileManager.h"
-#include "Misc/NotNull.h"
+#include "Types/SandboxInitFlags.h"
 
 namespace UE::FileSandboxCore
 {
@@ -15,6 +15,9 @@ struct FSandboxInitArgs
 	 * In theory, you could use this to save out the sandboxed files outside of the filesystem, e.g. to the cloud.
 	 */
 	IPlatformFile* OverridePlatformFile = nullptr;
+	
+	/** Additional feature flags for the sandbox. */
+	ESandboxInitFlags Flags = ESandboxInitFlags::Default;
 
 	/** @return The platform file to use. */
 	IPlatformFile& GetPlatformFile() const { return OverridePlatformFile ? *OverridePlatformFile : FPlatformFileManager::Get().GetPlatformFile(); }
diff --git a/Engine/Plugins/Developer/Sandbox/FileSandbox/Source/FileSandboxCore/Public/Types/SandboxInitFlags.h b/Engine/Plugins/Developer/Sandbox/FileSandbox/Source/FileSandboxCore/Public/Types/SandboxInitFlags.h
new file mode 100644
--- /dev/null
+++ b/Engine/Plugins/Developer/Sandbox/FileSandbox/Source/FileSandboxCore/Public/Types/SandboxInitFlags.h
@@ -0,0 +1,24 @@
+﻿// Copyright Epic Games, Inc. All Rights Reserved.
+
+#pragma once
+
+#include "HAL/Platform.h"
+#include "Misc/EnumClassFlags.h"
+
+namespace UE::FileSandboxCore
+{
+enum class ESandboxInitFlags : uint8
+{
+	None,
+	
+	/** 
+	 * If the user is using source control, then renaming assets should cause redirectors to be created. 
+	 * Under the hood, it forces ISourceControlState::IsLocal() to return false.
+	 */
+	ForceRedirectorsOnRenameInSourceControl = 1 << 0,
+	
+	Default = ForceRedirectorsOnRenameInSourceControl
+};
+
+ENUM_CLASS_FLAGS(ESandboxInitFlags);
+}
\ No newline at end of file
