File Sandbox: Minor directory utils API adjustments for UI.
- Made IsRootSandboxDirectory public
- Remove write flags for saving meta data

#Synced-CL 47637010

diff --git a/Engine/Plugins/Developer/Sandbox/FileSandbox/Source/FileSandboxCore/Private/Utils/FileSandboxUtils.cpp b/Engine/Plugins/Developer/Sandbox/FileSandbox/Source/FileSandboxCore/Private/Utils/FileSandboxUtils.cpp
--- a/Engine/Plugins/Developer/Sandbox/FileSandbox/Source/FileSandboxCore/Private/Utils/FileSandboxUtils.cpp
+++ b/Engine/Plugins/Developer/Sandbox/FileSandbox/Source/FileSandboxCore/Private/Utils/FileSandboxUtils.cpp
@@ -38,12 +38,6 @@ TOptional<FFileSandboxCore_ManifestData> LoadManifest(const FString& InRootDirec
 	return {};
 }
 
-bool IsRootSandboxDirectory(const FString& InRootDirectory)
-{
-	return FPaths::FileExists(InRootDirectory / GetManifestFileName())
-		&& FPaths::FileExists(InRootDirectory / GetMetadataFileName());
-}
-
 FString GetManifestFileName()
 {
 	return TEXT("manifest.json");
diff --git a/Engine/Plugins/Developer/Sandbox/FileSandbox/Source/FileSandboxCore/Private/Utils/FileSandboxUtils.h b/Engine/Plugins/Developer/Sandbox/FileSandbox/Source/FileSandboxCore/Private/Utils/FileSandboxUtils.h
--- a/Engine/Plugins/Developer/Sandbox/FileSandbox/Source/FileSandboxCore/Private/Utils/FileSandboxUtils.h
+++ b/Engine/Plugins/Developer/Sandbox/FileSandbox/Source/FileSandboxCore/Private/Utils/FileSandboxUtils.h
@@ -13,9 +13,6 @@ bool SaveManifest(const FFileSandboxCore_ManifestData& InData, const FString& In
 /** Loads the manifest file. */
 TOptional<FFileSandboxCore_ManifestData> LoadManifest(const FString& InRootDirectory, bool bLogErrorIfNotExist = true);
 
-/** @return Whether this is a path to the root of a sandbox, i.e. whether it contains the manifest and the metadata file. */
-bool IsRootSandboxDirectory(const FString& InRootDirectory);
-
 /* @return The name to give the manifest file. */
 FString GetManifestFileName();
 /* @return The name to give the metadata file. */
diff --git a/Engine/Plugins/Developer/Sandbox/FileSandbox/Source/FileSandboxCore/Private/Utils/SandboxDirectoryUtils.cpp b/Engine/Plugins/Developer/Sandbox/FileSandbox/Source/FileSandboxCore/Private/Utils/SandboxDirectoryUtils.cpp
--- a/Engine/Plugins/Developer/Sandbox/FileSandbox/Source/FileSandboxCore/Private/Utils/SandboxDirectoryUtils.cpp
+++ b/Engine/Plugins/Developer/Sandbox/FileSandbox/Source/FileSandboxCore/Private/Utils/SandboxDirectoryUtils.cpp
@@ -38,40 +38,20 @@ void ForEachSandbox(
 	});
 }
 
-namespace Private
+bool IsRootSandboxDirectory(const FString& InRootDirectory)
 {
-static FFileSandboxCore_SandboxMetaData MaskMetaData(const FFileSandboxCore_SandboxMetaData& InMetaData, EMetaDataWriteFlags InFlags)
-{
-	FFileSandboxCore_SandboxMetaData Result;
-
-	if (EnumHasAnyFlags(InFlags, EMetaDataWriteFlags::Name))
-	{
-		Result.Name = InMetaData.Name;
-	}
-	if (EnumHasAnyFlags(InFlags, EMetaDataWriteFlags::Description))
-	{
-		Result.Description = InMetaData.Description;
-	}
-	if (EnumHasAnyFlags(InFlags, EMetaDataWriteFlags::Tags))
-	{
-		Result.Tags = InMetaData.Tags;
-	}
-	if (EnumHasAnyFlags(InFlags, EMetaDataWriteFlags::CustomData))
-	{
-		Result.CustomData = InMetaData.CustomData;
-	}
-
-	return Result;
-}
+	return FPaths::DirectoryExists(InRootDirectory)
+		&& FPaths::FileExists(InRootDirectory / GetManifestFileName())
+		&& FPaths::FileExists(InRootDirectory / GetMetadataFileName());
 }
 
 bool SaveMetaData(
-	const FFileSandboxCore_SandboxMetaData& NewMetaData, const FString& InSandboxDirectory, EMetaDataWriteFlags InWriteFlags
+	const FFileSandboxCore_SandboxMetaData& NewMetaData, const FString& InSandboxDirectory
 	)
 {
 	FString JSONPayload;
 	FJsonObjectConverter::UStructToJsonObjectString(
-		EnumHasAllFlags(InWriteFlags, EMetaDataWriteFlags::All) ? NewMetaData : Private::MaskMetaData(NewMetaData, InWriteFlags),
+		NewMetaData,
 		JSONPayload, 0, 0
 		);
 	
diff --git a/Engine/Plugins/Developer/Sandbox/FileSandbox/Source/FileSandboxCore/Public/Utils/SandboxDirectoryUtils.h b/Engine/Plugins/Developer/Sandbox/FileSandbox/Source/FileSandboxCore/Public/Utils/SandboxDirectoryUtils.h
--- a/Engine/Plugins/Developer/Sandbox/FileSandbox/Source/FileSandboxCore/Public/Utils/SandboxDirectoryUtils.h
+++ b/Engine/Plugins/Developer/Sandbox/FileSandbox/Source/FileSandboxCore/Public/Utils/SandboxDirectoryUtils.h
@@ -32,21 +32,10 @@ UE_API void ForEachSandbox(
 	const FString& InBaseDirectory = GetBaseSandboxDirectory()
 	);
 
+/** @return Whether this is a path to the root of a sandbox, i.e. whether it contains the manifest and the metadata file. */
+UE_API bool IsRootSandboxDirectory(const FString& InRootDirectory);
 
-	
-enum class EMetaDataWriteFlags : uint8
-{
-	None,
-
-	Name = 1 << 0,
-	Description = 1 << 1,
-	Tags = 1 << 2,
-	CustomData = 1 << 3,
 
-	All = Name | Description | Tags | CustomData
-};
-ENUM_CLASS_FLAGS(EMetaDataWriteFlags);
-	
 /**
  * Writes new metadata for a sandbox.
  * 
@@ -55,11 +44,9 @@ ENUM_CLASS_FLAGS(EMetaDataWriteFlags);
  * 
  * @param NewMetaData The new metadata to set for the sandbox.
  * @param InSandboxDirectory Path to the base directory of the sandbox, e.g. what's returned by ForEachSandbox. Either absolute or relative to the process BaseDir().
- * @param InWriteFlags Optional. A property mask of properties you want to write. If you only wanted to write the Description, you can leave the other members blank.
  */
 UE_API bool SaveMetaData(
-	const FFileSandboxCore_SandboxMetaData& NewMetaData, const FString& InSandboxDirectory, 
-	EMetaDataWriteFlags InWriteFlags = EMetaDataWriteFlags::All
+	const FFileSandboxCore_SandboxMetaData& NewMetaData, const FString& InSandboxDirectory
 	);
 	
 /**
