File Sandbox Core: Add repository interface, which caches the sandboxes available on disk.

#Synced-CL 47637010

diff --git a/Engine/Plugins/Developer/Sandbox/FileSandbox/Source/FileSandboxCore/Private/FileSandboxCoreModule.cpp b/Engine/Plugins/Developer/Sandbox/FileSandbox/Source/FileSandboxCore/Private/FileSandboxCoreModule.cpp
--- a/Engine/Plugins/Developer/Sandbox/FileSandbox/Source/FileSandboxCore/Private/FileSandboxCoreModule.cpp
+++ b/Engine/Plugins/Developer/Sandbox/FileSandbox/Source/FileSandboxCore/Private/FileSandboxCoreModule.cpp
@@ -2,15 +2,29 @@
 
 #include "FileSandboxCoreModule.h"
 
+#include "Repository/NaiveSandboxRepository.h"
+#include "Utils/SandboxDirectoryUtils.h"
+
+#if UE_SANDBOX_WITH_DIRECTORY_WATCHER
+#include "Repository/WatchedSandboxRepository.h"
+#endif
+
 namespace UE::FileSandboxCore
 {
 void FFileSandboxCoreModule::StartupModule()
 {
 	SandboxManager = MakeUnique<FSandboxManager>();
+	
+#if UE_SANDBOX_WITH_DIRECTORY_WATCHER
+	DefaultRepository = MakeUnique<FWatchedSandboxRepository>(GetBaseSandboxDirectory());
+#else
+	DefaultRepository = MakeUnique<FNaiveSandboxRepository>(GetBaseSandboxDirectory(), *SandboxManager);
+#endif
 }
 
 void FFileSandboxCoreModule::ShutdownModule()
 {
+	DefaultRepository.Reset();
 	SandboxManager.Reset();
 }
 }
diff --git a/Engine/Plugins/Developer/Sandbox/FileSandbox/Source/FileSandboxCore/Private/FileSandboxCoreModule.h b/Engine/Plugins/Developer/Sandbox/FileSandbox/Source/FileSandboxCore/Private/FileSandboxCoreModule.h
--- a/Engine/Plugins/Developer/Sandbox/FileSandbox/Source/FileSandboxCore/Private/FileSandboxCoreModule.h
+++ b/Engine/Plugins/Developer/Sandbox/FileSandbox/Source/FileSandboxCore/Private/FileSandboxCoreModule.h
@@ -18,12 +18,16 @@ public:
 
 	//~ Begin IFileSandboxCoreModule Module
 	virtual ISandboxManager& GetSandboxManager() override { return *SandboxManager; }
+	virtual ISandboxRepository& GetDefaultSandboxRepository() override { return *DefaultRepository; }
 	//~ End IFileSandboxCoreModule Module
 
 private:
 
 	/** Public API for interacting with the sandbox instance. */
 	TUniquePtr<FSandboxManager> SandboxManager;
+	
+	/** Caches the sandboxes in the default sandbox directory. */
+	TUniquePtr<ISandboxRepository> DefaultRepository;
 };
 }
 
diff --git a/Engine/Plugins/Developer/Sandbox/FileSandbox/Source/FileSandboxCore/Private/Repository/NaiveSandboxRepository.cpp b/Engine/Plugins/Developer/Sandbox/FileSandbox/Source/FileSandboxCore/Private/Repository/NaiveSandboxRepository.cpp
new file mode 100644
--- /dev/null
+++ b/Engine/Plugins/Developer/Sandbox/FileSandbox/Source/FileSandboxCore/Private/Repository/NaiveSandboxRepository.cpp
@@ -0,0 +1,60 @@
+﻿// Copyright Epic Games, Inc. All Rights Reserved.
+
+#include "NaiveSandboxRepository.h"
+
+#include "ISandboxManager.h"
+#include "Types/EBreakBehavior.h"
+#include "Types/RepositoryChangedEvent.h"
+#include "Utils/SandboxDirectoryUtils.h"
+
+namespace UE::FileSandboxCore
+{
+FNaiveSandboxRepository::FNaiveSandboxRepository(FString InBaseDirectory, ISandboxManager& InManager)
+	: BaseDirectory(MoveTemp(InBaseDirectory))
+	, Manager(InManager)
+{
+	Manager.OnPostSandboxStartup().AddRaw(this, &FNaiveSandboxRepository::OnSandboxStartup);
+}
+
+FNaiveSandboxRepository::~FNaiveSandboxRepository()
+{
+	Manager.OnPostSandboxStartup().RemoveAll(this);
+}
+
+void FNaiveSandboxRepository::ForEachSandbox(TFunctionRef<EBreakBehavior(const FString& InRootPath, const FFileSandboxCore_SandboxMetaData&)> InProcessSandboxes)
+{
+	FileSandboxCore::ForEachSandbox([this, &InProcessSandboxes](const FString& InRootPath)
+	{
+		EBreakBehavior Behavior = EBreakBehavior::Continue;
+		ReadMetaData(InRootPath, [&InRootPath, &InProcessSandboxes, &Behavior](const FFileSandboxCore_SandboxMetaData& InMetaData)
+		{
+			Behavior = InProcessSandboxes(InRootPath, InMetaData);
+		});
+		return Behavior;
+	}
+	, BaseDirectory);
+}
+
+int32 FNaiveSandboxRepository::NumSandboxes() const
+{
+	int32 Num = 0;
+	FileSandboxCore::ForEachSandbox([&Num](auto){ ++Num; return EBreakBehavior::Continue; }, BaseDirectory);
+	return Num;
+}
+
+bool FNaiveSandboxRepository::ReadMetaData(const FString& InRootPath, TFunctionRef<void(const FFileSandboxCore_SandboxMetaData&)> InProcessMetadata)
+{
+	if (const TOptional<FFileSandboxCore_SandboxMetaData> MetaData = FileSandboxCore::LoadMetaData(InRootPath))
+	{
+		InProcessMetadata(*MetaData);
+		return true;
+	}
+	return false;
+}
+
+void FNaiveSandboxRepository::OnSandboxStartup(ISandboxInstance& SandboxInstance) const
+{
+	const FString RootDir(SandboxInstance.GetRootDirectory());
+	OnSandboxesChangedDelegate.Broadcast(FRepositoryChangedEvent{ .AddedSandboxPaths = MakeArrayView(&RootDir, 1) });
+}
+}
diff --git a/Engine/Plugins/Developer/Sandbox/FileSandbox/Source/FileSandboxCore/Private/Repository/NaiveSandboxRepository.h b/Engine/Plugins/Developer/Sandbox/FileSandbox/Source/FileSandboxCore/Private/Repository/NaiveSandboxRepository.h
new file mode 100644
--- /dev/null
+++ b/Engine/Plugins/Developer/Sandbox/FileSandbox/Source/FileSandboxCore/Private/Repository/NaiveSandboxRepository.h
@@ -0,0 +1,52 @@
+﻿// Copyright Epic Games, Inc. All Rights Reserved.
+
+#pragma once
+
+#include "ISandboxInstance.h"
+#include "Data/SandboxMetaData.h"
+#include "ISandboxRepository.h"
+#include "Templates/UnrealTemplate.h"
+#include "Utils/ScopedWatchedDirectory.h"
+
+struct FFileChangeData;
+
+namespace UE::FileSandboxCore
+{
+class ISandboxManager;
+
+/** 
+ * Naive implementation that always reads the disk and does no caching. Never invokes OnSandboxMetaDataChanged. 
+ * Used when directory watching is not possible.
+ */
+class FNaiveSandboxRepository : public ISandboxRepository, public FNoncopyable
+{
+public:
+	
+	explicit FNaiveSandboxRepository(FString InBaseDirectory, ISandboxManager& InManager UE_LIFETIMEBOUND);
+	~FNaiveSandboxRepository();
+
+	//~ Begin ISandboxRepository Interface
+	virtual void ForEachSandbox(TFunctionRef<EBreakBehavior(const FString& InRootPath, const FFileSandboxCore_SandboxMetaData&)> InProcessSandboxes) override;
+	virtual int32 NumSandboxes() const override;
+	virtual bool ReadMetaData(const FString& InRootPath, TFunctionRef<void(const FFileSandboxCore_SandboxMetaData&)> InProcessMetadata) override;
+	virtual FOnSandboxesChanged& OnSandboxesChanged() override { return OnSandboxesChangedDelegate; }
+	virtual FOnSandboxMetaDataChanged& OnSandboxMetaDataChanged() override { return OnSandboxMetaDataChangedDelegate; }
+	//~ End ISandboxRepository Interface
+
+private:
+	
+	/** The base directory in which to look for sandboxes. Only looks at the top-level, i.e. not recursively. */
+	const FString BaseDirectory;
+	
+	/** Sandbox manager. Used to detect when a sandbox is created. */
+	ISandboxManager& Manager;
+	
+	/** Invoked when the known sandboxes changes (removed, discovered, etc.) */
+	FOnSandboxesChanged OnSandboxesChangedDelegate;
+	/** Invoked when a sandbox's metadata changes. */
+	FOnSandboxMetaDataChanged OnSandboxMetaDataChangedDelegate;
+	
+	/** Callback when a sandbox is started. */
+	void OnSandboxStartup(ISandboxInstance& SandboxInstance) const;
+};	
+}
\ No newline at end of file
diff --git a/Engine/Plugins/Developer/Sandbox/FileSandbox/Source/FileSandboxCore/Private/Repository/WatchedSandboxRepository.cpp b/Engine/Plugins/Developer/Sandbox/FileSandbox/Source/FileSandboxCore/Private/Repository/WatchedSandboxRepository.cpp
new file mode 100644
--- /dev/null
+++ b/Engine/Plugins/Developer/Sandbox/FileSandbox/Source/FileSandboxCore/Private/Repository/WatchedSandboxRepository.cpp
@@ -0,0 +1,258 @@
+﻿// Copyright Epic Games, Inc. All Rights Reserved.
+
+#if UE_SANDBOX_WITH_DIRECTORY_WATCHER
+#include "WatchedSandboxRepository.h"
+
+#include "IDirectoryWatcher.h"
+#include "IDirectoryWatcher.h"
+#include "Algo/AnyOf.h"
+#include "GenericPlatform/GenericPlatformFile.h"
+#include "HAL/PlatformFileManager.h"
+#include "Misc/Paths.h"
+#include "Types/EBreakBehavior.h"
+#include "Types/RepositoryChangedEvent.h"
+#include "Utils/FileSandboxUtils.h"
+#include "Utils/SandboxDirectoryUtils.h"
+#include "Utils/SandboxUtils.h"
+
+namespace UE::FileSandboxCore
+{
+FWatchedSandboxRepository::FWatchedSandboxRepository(FString InBaseDirectory)
+	: BaseDirectory(MoveTemp(InBaseDirectory))
+	, WatchedBaseDirectory([this]()
+	{
+		if (IDirectoryWatcher* DirectoryWatcher = GetDirectoryWatcher())
+		{
+			FDelegateHandle Handle;
+			DirectoryWatcher->RegisterDirectoryChangedCallback_Handle(
+				BaseDirectory, 
+				IDirectoryWatcher::FDirectoryChanged::CreateRaw(this, &FWatchedSandboxRepository::OnBaseDirectoryChanged),
+				Handle,
+				// We only care for the direct directories. Ignore stuff in the sandbox sub-directories.
+				IDirectoryWatcher::WatchOptions::IncludeDirectoryChanges | IDirectoryWatcher::WatchOptions::IgnoreChangesInSubtree
+				);
+			return FScopedWatchedDirectory(BaseDirectory, Handle);
+		}
+			
+		return FScopedWatchedDirectory{};
+	}())
+{
+	RescanBaseDirectory();
+	UpdateWatchedDirectories();
+}
+
+void FWatchedSandboxRepository::ForEachSandbox(TFunctionRef<EBreakBehavior(const FString& InRootPath, const FFileSandboxCore_SandboxMetaData& MetaData)> InProcessSandboxes)
+{
+	for (const FSandboxData& SandboxData : CachedSandboxData)
+	{
+		if (InProcessSandboxes(SandboxData.AbsPath, SandboxData.MetaData) == EBreakBehavior::Break)
+		{
+			break;
+		}
+	}
+}
+
+bool FWatchedSandboxRepository::ReadMetaData(const FString& InRootPath, TFunctionRef<void(const FFileSandboxCore_SandboxMetaData&)> InProcessMetadata)
+{
+	const int32 Index = IndexOf(InRootPath);
+	if (CachedSandboxData.IsValidIndex(Index))
+	{
+		InProcessMetadata(CachedSandboxData[Index].MetaData);
+		return true;
+	}
+	return false;
+}
+
+void FWatchedSandboxRepository::RescanBaseDirectory()
+{
+	CachedSandboxData.Reset();
+	
+	FileSandboxCore::ForEachSandbox([this](const FString& InBaseDirectory)
+	{
+		AddSandboxIfExists(InBaseDirectory);
+		return EBreakBehavior::Continue;
+	}, BaseDirectory);
+}
+
+void FWatchedSandboxRepository::AddSandboxIfExists(const FString& InDirectory)
+{
+	if (!IsRootSandboxDirectory(InDirectory))
+	{
+		return;
+	}
+	
+	const FString AbsPath = FPaths::ConvertRelativePathToFull(InDirectory);
+	TOptional<FFileSandboxCore_SandboxMetaData> MetaData = LoadMetaData(AbsPath);
+	if (!MetaData)
+	{
+		return;
+	}
+	
+	const int32 ExistingIndex = IndexOf(AbsPath);
+	if (CachedSandboxData.IsValidIndex(ExistingIndex)) // This is just defensive code - this case should not occur.
+	{
+		CachedSandboxData[ExistingIndex].MetaData = MoveTemp(*MetaData);
+		OnSandboxMetaDataChangedDelegate.Broadcast(AbsPath);
+	}
+	else
+	{
+		CachedSandboxData.Emplace(FSandboxData{ AbsPath, MoveTemp(*MetaData) });
+		OnSandboxesChanged().Broadcast(FRepositoryChangedEvent{ .AddedSandboxPaths = MakeArrayView(&AbsPath, 1) });
+	}
+}
+
+int32 FWatchedSandboxRepository::IndexOf(const FString& InDirectory) const
+{
+	return CachedSandboxData.IndexOfByPredicate([&InDirectory](const FSandboxData& Data)
+	{
+		return FPaths::IsSamePath(Data.AbsPath, InDirectory);
+	});
+}
+
+void FWatchedSandboxRepository::OnBaseDirectoryChanged(const TArray<FFileChangeData>& InChanges)
+{
+	check(IsInGameThread());
+	
+	UpdateCacheFromChangedBaseDirectory(InChanges);
+	UpdateWatchedDirectories();
+}
+
+void FWatchedSandboxRepository::OnSandboxDirectoryChanged(const TArray<FFileChangeData>& InChanges, FString InSandboxPath)
+{
+	check(IsInGameThread());
+	
+	const int32 Index = IndexOf(InSandboxPath);
+	
+	const bool bWasSandboxDirectoryBefore = CachedSandboxData.IsValidIndex(Index);
+	const bool bIsSandboxDirectory = IsRootSandboxDirectory(InSandboxPath);
+	if (bIsSandboxDirectory && !bWasSandboxDirectoryBefore)
+	{
+		AddSandboxIfExists(InSandboxPath);
+		return;
+	}
+	
+	if (!bIsSandboxDirectory)
+	{
+		CachedSandboxData.RemoveAt(Index);
+		OnSandboxesChanged().Broadcast(FRepositoryChangedEvent{ .RemovedSandboxPaths = MakeArrayView(&InSandboxPath, 1) });
+		return;
+	}
+	
+	DetectMetaDataFileChange(Index, InChanges, InSandboxPath);
+}
+
+void FWatchedSandboxRepository::DetectMetaDataFileChange(int32 InSandboxIndex, const TArray<FFileChangeData>& InChanges, const FString& InSandboxPath)
+{
+	const bool bContainsUnknownOrRescan = InChanges.ContainsByPredicate([](const FFileChangeData& Change)
+	{
+		return Change.Action == FFileChangeData::FCA_RescanRequired || Change.Action == FFileChangeData::FCA_Unknown;
+	});
+	const int32 MetadataFileChangeIndex = InChanges.IndexOfByPredicate([](const FFileChangeData& Change)
+	{
+		return FPaths::GetCleanFilename(Change.Filename) == GetMetadataFileName();
+	});
+	// If the directory watcher is unsure what happened (rescan / unknown), we'll treat the file as changed, i.e. parse it again.
+	if (!bContainsUnknownOrRescan && !InChanges.IsValidIndex(MetadataFileChangeIndex))
+	{
+		return;
+	}
+	
+	if (TOptional<FFileSandboxCore_SandboxMetaData> MetaData = LoadMetaData(InSandboxPath))
+	{
+		CachedSandboxData[InSandboxIndex].MetaData = MoveTemp(*MetaData);
+	}
+	else
+	{
+		CachedSandboxData.RemoveAtSwap(InSandboxIndex);
+		OnSandboxesChangedDelegate.Broadcast(FRepositoryChangedEvent{ .RemovedSandboxPaths = MakeArrayView(&InSandboxPath, 1) });
+	}
+}
+
+void FWatchedSandboxRepository::UpdateCacheFromChangedBaseDirectory(const TArray<FFileChangeData>& InChanges)
+{
+	const bool bDoFullRescan = Algo::AnyOf(InChanges, [](const FFileChangeData& Change)
+	{
+		const bool bIsDirectoryChange = FPaths::DirectoryExists(Change.Filename);
+		return bIsDirectoryChange && (Change.Action == FFileChangeData::FCA_Unknown || Change.Action == FFileChangeData::FCA_RescanRequired);
+	});
+	if (bDoFullRescan)
+	{
+		RescanBaseDirectory();
+		return;
+	}
+	
+	TSet<FString> AlreadyProcessed;
+	for (const FFileChangeData& Change : InChanges)
+	{
+		const FString& Path = Change.Filename;
+		
+		// InChanges can contain 100s of FCA_Modified changes for the same file. Don't do double work.
+		if (AlreadyProcessed.Contains(Path))
+		{
+			continue;
+		}
+		AlreadyProcessed.Add(Path);
+
+		switch (Change.Action)
+		{
+		case FFileChangeData::FCA_Added: AddSandboxIfExists(Path); break;
+		case FFileChangeData::FCA_Removed: 
+			if (const int32 Index = IndexOf(Path))
+			{
+				CachedSandboxData.RemoveAtSwap(Index);
+				OnSandboxesChanged().Broadcast(FRepositoryChangedEvent{ .RemovedSandboxPaths = MakeArrayView(&Path, 1) });
+			}
+			break;
+			
+		case FFileChangeData::FCA_Modified: break; // Don't care. The sandbox name is what it says in the metadata file, not the directory name.
+		default: checkNoEntry();
+		}
+	}
+}
+
+void FWatchedSandboxRepository::UpdateWatchedDirectories()
+{
+	IDirectoryWatcher* DirectoryWatcher = GetDirectoryWatcher();
+	if (!DirectoryWatcher)
+	{
+		return;
+	}
+	
+	// Reuse old delegates if possible
+	TArray<FScopedWatchedDirectory> Old = MoveTemp(WatchedSubDirectories);
+	
+	IPlatformFile& PlatformFile = FPlatformFileManager::Get().GetPlatformFile();
+	PlatformFile.IterateDirectory(*BaseDirectory, [this, &Old, DirectoryWatcher](const TCHAR* FilenameOrDirectory, const bool bIsDirectory)
+	{
+		if (!bIsDirectory)
+		{
+			return true;
+		}
+			
+		const int32 ExistingIndex = Old.IndexOfByPredicate([&FilenameOrDirectory](const FScopedWatchedDirectory& WatchedDirectory)
+		{
+			return FPaths::IsSamePath(WatchedDirectory.GetWatchedDirectory(), FilenameOrDirectory);
+		});
+			
+		if (Old.IsValidIndex(ExistingIndex))
+		{
+			WatchedSubDirectories.Emplace(MoveTemp(Old[ExistingIndex]));
+		}
+		else
+		{
+			FDelegateHandle Handle;
+			DirectoryWatcher->RegisterDirectoryChangedCallback_Handle(
+				FilenameOrDirectory, 
+				IDirectoryWatcher::FDirectoryChanged::CreateRaw(this, &FWatchedSandboxRepository::OnSandboxDirectoryChanged, FString(FilenameOrDirectory)),
+				Handle,
+				IDirectoryWatcher::WatchOptions::IgnoreChangesInSubtree // We only care for the manifest and metadata file
+				);
+			WatchedSubDirectories.Emplace(FilenameOrDirectory, Handle);
+		}
+			
+		return true;
+	});
+}
+}
+
+#endif
\ No newline at end of file
diff --git a/Engine/Plugins/Developer/Sandbox/FileSandbox/Source/FileSandboxCore/Private/Repository/WatchedSandboxRepository.h b/Engine/Plugins/Developer/Sandbox/FileSandbox/Source/FileSandboxCore/Private/Repository/WatchedSandboxRepository.h
new file mode 100644
--- /dev/null
+++ b/Engine/Plugins/Developer/Sandbox/FileSandbox/Source/FileSandboxCore/Private/Repository/WatchedSandboxRepository.h
@@ -0,0 +1,80 @@
+﻿// Copyright Epic Games, Inc. All Rights Reserved.
+
+#pragma once
+
+#if UE_SANDBOX_WITH_DIRECTORY_WATCHER
+#include "ISandboxRepository.h"
+#include "Data/SandboxMetaData.h"
+#include "Templates/UnrealTemplate.h"
+#include "Utils/ScopedWatchedDirectory.h"
+
+struct FFileChangeData;
+
+namespace UE::FileSandboxCore
+{
+/** Simple directory watcher that detects when sandboxes are added, removed, and metadata file edited. */
+class FWatchedSandboxRepository : public ISandboxRepository, public FNoncopyable
+{
+public:
+
+	explicit FWatchedSandboxRepository(FString InBaseDirectory);
+
+	//~ Begin ISandboxRepository Interface
+	virtual void ForEachSandbox(TFunctionRef<EBreakBehavior(const FString& InRootPath, const FFileSandboxCore_SandboxMetaData& MetaData)> InProcessSandboxes) override;
+	virtual int32 NumSandboxes() const override { return CachedSandboxData.Num(); }
+	virtual bool ReadMetaData(const FString& InRootPath, TFunctionRef<void(const FFileSandboxCore_SandboxMetaData&)> InProcessMetadata) override;
+	virtual FOnSandboxesChanged& OnSandboxesChanged() override { return OnSandboxesChangedDelegate; }
+	virtual FOnSandboxMetaDataChanged& OnSandboxMetaDataChanged() override { return OnSandboxMetaDataChangedDelegate; }
+	//~ End ISandboxRepository Interface
+
+private:
+	
+	/** The base directory in which to look for sandboxes. Only looks at the top-level, i.e. not recursively. */
+	const FString BaseDirectory;
+	
+	struct FSandboxData
+	{
+		/** Absolute path to the sandbox root directory, i.e. the directory containing the metadata file. */
+		FString AbsPath;
+		/** Cached content of the sandbox metadata file (to avoid constant JSON parsing, file loading, etc.) */
+		FFileSandboxCore_SandboxMetaData MetaData;
+	};
+	/** Cached sandbox data. Updated whenever info changes on disk. */
+	TArray<FSandboxData> CachedSandboxData;
+	
+	/** Watcher for BaseDirectory */
+	const FScopedWatchedDirectory WatchedBaseDirectory;
+	/** 
+	 * Sub-directories were watching in WatchedBaseDirectory. 
+	 * All directories in the base directory are always watched so we can detect when a new sandbox is created there. 
+	 */
+	TArray<FScopedWatchedDirectory> WatchedSubDirectories;
+	
+	/** Invoked when the known sandboxes changes (removed, discovered, etc.) */
+	FOnSandboxesChanged OnSandboxesChangedDelegate;
+	/** Invoked when a sandbox's metadata changes. */
+	FOnSandboxMetaDataChanged OnSandboxMetaDataChangedDelegate;
+	
+	/** Fully rebuilds all directory data, analzying all directory, parsing all metadata files, etc. */
+	void RescanBaseDirectory();
+	/** Adds CachedSandboxData entry if it contains a sandbox. Also starts listening for file changes. */
+	void AddSandboxIfExists(const FString& InDirectory);
+	
+	/** @return Index to CachedSandboxData entry that contains the InDirectory (relative or absolute). */
+	int32 IndexOf(const FString& InDirectory) const;
+	
+	/** Invoked when a change is made in the base directory. */
+	void OnBaseDirectoryChanged(const TArray<FFileChangeData>& InChanges);
+	/** Checks all reported changes made to directories with the goal of detecting added or removed sandboxes. */
+	void UpdateCacheFromChangedBaseDirectory(const TArray<FFileChangeData>& InChanges);
+	/** Ensures all direct sub-directories in the base directory have a directory watcher set up. */
+	void UpdateWatchedDirectories();
+	
+	/** Invoked when a change is made in any of the sandbox directories. Used to update metadata. */
+	void OnSandboxDirectoryChanged(const TArray<FFileChangeData>& InChanges, FString InSandboxPath);
+	/** Checks whether a change has been reported on the metadata file and invokes OnSandboxMetaDataChangedDelegate, if so. */
+	void DetectMetaDataFileChange(int32 InSandboxIndex, const TArray<FFileChangeData>& InChanges, const FString& InSandboxPath);
+};	
+}
+
+#endif
\ No newline at end of file
diff --git a/Engine/Plugins/Developer/Sandbox/FileSandbox/Source/FileSandboxCore/Private/Utils/ScopedWatchedDirectory.cpp b/Engine/Plugins/Developer/Sandbox/FileSandbox/Source/FileSandboxCore/Private/Utils/ScopedWatchedDirectory.cpp
new file mode 100644
--- /dev/null
+++ b/Engine/Plugins/Developer/Sandbox/FileSandbox/Source/FileSandboxCore/Private/Utils/ScopedWatchedDirectory.cpp
@@ -0,0 +1,37 @@
+﻿// Copyright Epic Games, Inc. All Rights Reserved.
+
+#if UE_SANDBOX_WITH_DIRECTORY_WATCHER
+
+#include "ScopedWatchedDirectory.h"
+
+#include "IDirectoryWatcher.h"
+#include "SandboxUtils.h"
+
+namespace UE::FileSandboxCore
+{
+FScopedWatchedDirectory::FScopedWatchedDirectory(FString InDirectory, FDelegateHandle InHandle)
+	: WatchedDirectory(InDirectory)
+	, Handle(InHandle)
+{}
+
+FScopedWatchedDirectory::FScopedWatchedDirectory(FScopedWatchedDirectory&& Old)
+	: FScopedWatchedDirectory(MoveTemp(Old.WatchedDirectory), Old.Handle)
+{
+	Old.Handle.Reset();
+}
+
+FScopedWatchedDirectory::~FScopedWatchedDirectory()
+{
+	if (!Handle.IsValid() || !WatchedDirectory.IsEmpty())
+	{
+		return;
+	}
+	
+	if (IDirectoryWatcher* Watcher = GetDirectoryWatcherIfLoaded())
+	{
+		Watcher->UnregisterDirectoryChangedCallback_Handle(WatchedDirectory, Handle);
+	}
+}
+}
+
+#endif
\ No newline at end of file
diff --git a/Engine/Plugins/Developer/Sandbox/FileSandbox/Source/FileSandboxCore/Private/Utils/ScopedWatchedDirectory.h b/Engine/Plugins/Developer/Sandbox/FileSandbox/Source/FileSandboxCore/Private/Utils/ScopedWatchedDirectory.h
new file mode 100644
--- /dev/null
+++ b/Engine/Plugins/Developer/Sandbox/FileSandbox/Source/FileSandboxCore/Private/Utils/ScopedWatchedDirectory.h
@@ -0,0 +1,32 @@
+﻿// Copyright Epic Games, Inc. All Rights Reserved.
+
+#pragma once
+
+#if UE_SANDBOX_WITH_DIRECTORY_WATCHER
+
+#include "Containers/UnrealString.h"
+#include "Delegates/IDelegateInstance.h"
+#include "Templates/UnrealTemplate.h"
+
+namespace UE::FileSandboxCore
+{
+/** Unregisters IDirectoryWatcher when it goes out of scope. */
+class FScopedWatchedDirectory : public FNoncopyable
+{
+public:
+	
+	FScopedWatchedDirectory() = default;
+	FScopedWatchedDirectory(FString InDirectory, FDelegateHandle InHandle);
+	FScopedWatchedDirectory(FScopedWatchedDirectory&& Old);
+	~FScopedWatchedDirectory();
+	
+	const FString& GetWatchedDirectory() const { return WatchedDirectory; }
+
+private:
+	
+	FString WatchedDirectory;
+	FDelegateHandle Handle;
+};
+}
+
+#endif
\ No newline at end of file
diff --git a/Engine/Plugins/Developer/Sandbox/FileSandbox/Source/FileSandboxCore/Public/IFileSandboxCoreModule.h b/Engine/Plugins/Developer/Sandbox/FileSandbox/Source/FileSandboxCore/Public/IFileSandboxCoreModule.h
--- a/Engine/Plugins/Developer/Sandbox/FileSandbox/Source/FileSandboxCore/Public/IFileSandboxCoreModule.h
+++ b/Engine/Plugins/Developer/Sandbox/FileSandbox/Source/FileSandboxCore/Public/IFileSandboxCoreModule.h
@@ -7,10 +7,9 @@
 
 namespace UE::FileSandboxCore
 {
-class ISandboxBrowser;
 class ISandboxInstance;
 class ISandboxManager;
-struct FSandboxInitArgs;
+class ISandboxRepository;
 
 class IFileSandboxCoreModule : public IModuleInterface
 {
@@ -33,5 +32,10 @@ public:
 	
 	/** @return The API for managing live sandboxes. */
 	virtual ISandboxManager& GetSandboxManager() = 0;
+	
+	/** @return Gets the repository for the default sandbox save location (@see GetBaseSandboxDirectory) */
+	virtual ISandboxRepository& GetDefaultSandboxRepository() = 0; 
+	
+	// In the future, we may consider adding CreateSandboxRepository(FString BasePath), or RegisterCustomSandboxRepository, to handle custom locations
 };
 }
diff --git a/Engine/Plugins/Developer/Sandbox/FileSandbox/Source/FileSandboxCore/Public/ISandboxRepository.h b/Engine/Plugins/Developer/Sandbox/FileSandbox/Source/FileSandboxCore/Public/ISandboxRepository.h
new file mode 100644
--- /dev/null
+++ b/Engine/Plugins/Developer/Sandbox/FileSandbox/Source/FileSandboxCore/Public/ISandboxRepository.h
@@ -0,0 +1,56 @@
+﻿// Copyright Epic Games, Inc. All Rights Reserved.
+
+#pragma once
+
+#include "Delegates/Delegate.h"
+#include "HAL/Platform.h"
+
+struct FFileSandboxCore_SandboxMetaData;
+
+namespace UE::FileSandboxCore
+{
+enum class EBreakBehavior : uint8;
+struct FRepositoryChangedEvent;
+	
+DECLARE_MULTICAST_DELEGATE_OneParam(FOnSandboxesChanged, const FRepositoryChangedEvent&);
+DECLARE_MULTICAST_DELEGATE_OneParam(FOnSandboxMetaDataChanged, const FString& /* InRootPath */);
+
+/** 
+ * Knows of the sandboxes instances available. 
+ * 
+ * You can think of this as a cache that will notify when the sandboxes available changes, and the metadata changes.
+ * You should prefer this over manually loading metadata (@see LoadMetaData) whenever possible to avoid excessive file reads, JSON parsing, etc.
+ * This is useful for, e.g. UI want to display the available sandboxes. 
+ */
+class ISandboxRepository
+{
+public:
+	
+	/** 
+	 * Enumerates all known sandboxes. 
+	 * InRootPath will be the absolute path to the root sandbox directory, i.e. the one containing the metadata file.
+	 */
+	virtual void ForEachSandbox(TFunctionRef<EBreakBehavior(const FString& InRootPath, const FFileSandboxCore_SandboxMetaData& MetaData)> InProcessSandboxes) = 0;
+	
+	/** @return Number of distinct sandboxes, i.e. number of times ForEachSandbox would invoke the callback. */
+	virtual int32 NumSandboxes() const = 0;
+	
+	
+	/**
+	 * @param InRootPath Relative or absolute path to the root directory of the sandbox, i.e. one enumerated by InProcessSandboxes. 
+	 * @param InProcessMetadata Callback invoked if the metadata exists
+	 * @return Whether the callback was invoked
+	 */
+	virtual bool ReadMetaData(const FString& InRootPath, TFunctionRef<void(const FFileSandboxCore_SandboxMetaData&)> InProcessMetadata) = 0;
+	
+	
+	/** @return Invoked when the known sandboxes changes (removed, discovered, etc.), on game thread. */
+	virtual FOnSandboxesChanged& OnSandboxesChanged() = 0;
+	
+	/** @return Invoked when a sandbox's metadata changes, on game thread. */
+	virtual FOnSandboxMetaDataChanged& OnSandboxMetaDataChanged() = 0;
+	
+	
+	virtual ~ISandboxRepository() = default;
+};
+}
diff --git a/Engine/Plugins/Developer/Sandbox/FileSandbox/Source/FileSandboxCore/Public/Types/RepositoryChangedEvent.h b/Engine/Plugins/Developer/Sandbox/FileSandbox/Source/FileSandboxCore/Public/Types/RepositoryChangedEvent.h
new file mode 100644
--- /dev/null
+++ b/Engine/Plugins/Developer/Sandbox/FileSandbox/Source/FileSandboxCore/Public/Types/RepositoryChangedEvent.h
@@ -0,0 +1,19 @@
+﻿// Copyright Epic Games, Inc. All Rights Reserved.
+
+#pragma once
+
+#include "Containers/UnrealString.h"
+#include "Containers/ArrayView.h"
+
+namespace UE::FileSandboxCore
+{
+/** Data describing a change in a repository */
+struct FRepositoryChangedEvent
+{
+	/** Root directories of sandboxes that have been added. */
+	TConstArrayView<FString> AddedSandboxPaths;
+	
+	/** Root directories of sandboxes that have been removed. */
+	TConstArrayView<FString> RemovedSandboxPaths;
+};
+}
