File Sandbox Core: Allow source control operations when sandbox is disabled.

#Synced-CL 47637010

diff --git a/Engine/Plugins/Developer/Sandbox/FileSandbox/Source/FileSandboxCore/Private/SourceControl/Proxy/SandboxSourceControlProviderProxy.cpp b/Engine/Plugins/Developer/Sandbox/FileSandbox/Source/FileSandboxCore/Private/SourceControl/Proxy/SandboxSourceControlProviderProxy.cpp
--- a/Engine/Plugins/Developer/Sandbox/FileSandbox/Source/FileSandboxCore/Private/SourceControl/Proxy/SandboxSourceControlProviderProxy.cpp
+++ b/Engine/Plugins/Developer/Sandbox/FileSandbox/Source/FileSandboxCore/Private/SourceControl/Proxy/SandboxSourceControlProviderProxy.cpp
@@ -57,6 +57,12 @@ static TOptional<ECommandResult::Type> InterceptSourceControlOp(
 	const FSourceControlOperationRef& InOperation , const TArray<FString>& InFiles
 	)
 {
+	// Sandbox becomes temporarily disabled when we're persisting changes.
+	if (!IFileManager::Get().IsSandboxEnabled())
+	{
+		return {};
+	}
+	
 	// Handle forwarding equivalents of edit, add, and remove to sandbox.
 	if (const TOptional<ECommandResult::Type> Result = HandleSandboxedOp(InOperation, InFiles))
 	{
