File Sandbox Core: Fix warning being issue when sandbox directory does not exist, yet.

#Synced-CL 47637010

diff --git a/Engine/Plugins/Developer/Sandbox/FileSandbox/Source/FileSandboxCore/Private/Repository/WatchedSandboxRepository.cpp b/Engine/Plugins/Developer/Sandbox/FileSandbox/Source/FileSandboxCore/Private/Repository/WatchedSandboxRepository.cpp
--- a/Engine/Plugins/Developer/Sandbox/FileSandbox/Source/FileSandboxCore/Private/Repository/WatchedSandboxRepository.cpp
+++ b/Engine/Plugins/Developer/Sandbox/FileSandbox/Source/FileSandboxCore/Private/Repository/WatchedSandboxRepository.cpp
@@ -5,6 +5,7 @@
 
 #include "IDirectoryWatcher.h"
 #include "IDirectoryWatcher.h"
+#include "LogFileSandbox.h"
 #include "Algo/AnyOf.h"
 #include "GenericPlatform/GenericPlatformFile.h"
 #include "HAL/PlatformFileManager.h"
@@ -21,7 +22,20 @@ FWatchedSandboxRepository::FWatchedSandboxRepository(FString InBaseDirectory)
 	: BaseDirectory(MoveTemp(InBaseDirectory))
 	, WatchedBaseDirectory([this]()
 	{
-		if (IDirectoryWatcher* DirectoryWatcher = GetDirectoryWatcher())
+		IDirectoryWatcher* DirectoryWatcher = GetDirectoryWatcher();
+		if (!DirectoryWatcher)
+		{
+			return FScopedWatchedDirectory{};;
+		}
+			
+		// Would be etter to wait until a user action causes directory to be created. 
+		// However, it's easier implementation-wise to always create this folder.
+		IPlatformFile& PlatformFile = FPlatformFileManager::Get().GetPlatformFile();
+		const bool bDirectoryExistedAlready = PlatformFile.DirectoryExists(*BaseDirectory);
+		const bool bHasCreatedDirectory = !bDirectoryExistedAlready && PlatformFile.CreateDirectoryTree(*BaseDirectory);
+			
+		const bool bCanSetupWatcher = bHasCreatedDirectory || bDirectoryExistedAlready;
+		if (bCanSetupWatcher)
 		{
 			FDelegateHandle Handle;
 			DirectoryWatcher->RegisterDirectoryChangedCallback_Handle(
