File Sandbox Core: Add engine version to metadata file

#Synced-CL 47637010

diff --git a/Engine/Plugins/Developer/Sandbox/FileSandbox/Source/FileSandboxCore/Private/Data/ManifestData.h b/Engine/Plugins/Developer/Sandbox/FileSandbox/Source/FileSandboxCore/Private/Data/ManifestData.h
--- a/Engine/Plugins/Developer/Sandbox/FileSandbox/Source/FileSandboxCore/Private/Data/ManifestData.h
+++ b/Engine/Plugins/Developer/Sandbox/FileSandbox/Source/FileSandboxCore/Private/Data/ManifestData.h
@@ -3,9 +3,13 @@
 #pragma once
 
 #include "AbsoluteFilePathArray.h"
+#include "Data/VersionInfo.h"
 #include "ManifestData.generated.h"
 
-/** Internal bookkeeping data for the sandbox logic. */
+/** 
+ * Internal bookkeeping data for the sandbox logic.
+ * This data is owned by the sandbox system and should not be directly modified via public API.
+ */
 USTRUCT()
 struct FFileSandboxCore_ManifestData
 {
@@ -22,6 +26,19 @@ struct FFileSandboxCore_ManifestData
 	/** Absolute non-sandbox paths to files added in sandbox. */
 	UPROPERTY()
 	FFileSandboxCore_AbsoluteFilePathArray AddedFiles;
+	
+	/**
+	 * Version info when this sandbox was created.
+	 * Holds info about the file version, engine version, and custom versions (for FArchives).
+	 * 
+	 * This information is saved in the manifest to minimize the risk of API users directly editing and corrupting it.
+	 * It could also have been saved in the metadata file but that has API for mutating it (= risk of corruption / mutation via API).
+	 * 
+	 * May be needed to parse the content in FFileSandboxCore_SandboxMetaData. 
+	 * TODO UE-350242: This information can also be used to reject loading a sandbox with newer version info than the current engine version.
+	 */
+	UPROPERTY()
+	FFileSandboxCore_VersionInfo VersionInfo;
 
 	void Empty()
 	{
diff --git a/Engine/Plugins/Developer/Sandbox/FileSandbox/Source/FileSandboxCore/Private/Data/VersionInfo.cpp b/Engine/Plugins/Developer/Sandbox/FileSandbox/Source/FileSandboxCore/Private/Data/VersionInfo.cpp
new file mode 100644
--- /dev/null
+++ b/Engine/Plugins/Developer/Sandbox/FileSandbox/Source/FileSandboxCore/Private/Data/VersionInfo.cpp
@@ -0,0 +1,64 @@
+﻿// Copyright Epic Games, Inc. All Rights Reserved.
+
+#include "Data/VersionInfo.h"
+
+#include "Containers/UnrealString.h"
+#include "Misc/EngineVersion.h"
+#include "Serialization/CustomVersion.h"
+
+void FFileSandboxCore_FileVersionInfo::Initialize()
+{
+	FileVersionUE5 = GPackageFileUEVersion.FileVersionUE5;
+	FileVersionLicensee = GPackageFileLicenseeUEVersion;
+}
+
+FString FFileSandboxCore_FileVersionInfo::ToString() const
+{
+	return FString::Printf(TEXT("FileVersionUE5=%d Licensee=%d"), FileVersionUE5, FileVersionLicensee);
+}
+
+void FFileSandboxCore_EngineVersionInfo::Initialize(const FEngineVersion& InVersion)
+{
+	Major = InVersion.GetMajor();
+	Minor = InVersion.GetMinor();
+	Patch = InVersion.GetPatch();
+	Changelist = InVersion.GetChangelist();
+}
+
+FString FFileSandboxCore_EngineVersionInfo::ToString(bool bIncludeChangelist) const
+{
+	return bIncludeChangelist 
+		? FString::Printf(TEXT("%u.%u.%u-%u"), Major, Minor, Patch, Changelist)
+		: FString::Printf(TEXT("%u.%u.%u"), Major, Minor, Patch);
+}
+
+void FFileSandboxCore_CustomVersionInfo::Initialize(const FCustomVersion& InVersion)
+{
+	FriendlyName = InVersion.GetFriendlyName();
+	Key = InVersion.Key;
+	Version = InVersion.Version;
+}
+
+void FFileSandboxCore_VersionInfo::Initialize()
+{
+	FileVersion.Initialize();
+	EngineVersion.Initialize(FEngineVersion::Current());
+
+	CustomVersions.Empty(CustomVersions.Num());
+	FCustomVersionContainer AllCurrentVersions = FCurrentCustomVersions::GetAll();
+	for (const FCustomVersion& EngineCustomVersion : AllCurrentVersions.GetAllVersions())
+	{
+		FFileSandboxCore_CustomVersionInfo& CustomVersion = CustomVersions.AddDefaulted_GetRef();
+		CustomVersion.Initialize(EngineCustomVersion);
+	}
+}
+
+bool FFileSandboxCore_VersionInfo::IsInitialized() const
+{
+	return CustomVersions.Num() > 0;
+}
+
+FString FFileSandboxCore_VersionInfo::ToString() const
+{
+	return FString::Printf(TEXT("%s %s"), *EngineVersion.ToString(), *FileVersion.ToString());
+}
diff --git a/Engine/Plugins/Developer/Sandbox/FileSandbox/Source/FileSandboxCore/Private/Sandbox/SandboxInstance.cpp b/Engine/Plugins/Developer/Sandbox/FileSandbox/Source/FileSandboxCore/Private/Sandbox/SandboxInstance.cpp
--- a/Engine/Plugins/Developer/Sandbox/FileSandbox/Source/FileSandboxCore/Private/Sandbox/SandboxInstance.cpp
+++ b/Engine/Plugins/Developer/Sandbox/FileSandbox/Source/FileSandboxCore/Private/Sandbox/SandboxInstance.cpp
@@ -34,6 +34,8 @@ TUniquePtr<FSandboxInstance> FSandboxInstance::CreateNewSandbox(
 	)
 {
 	FFileSandboxCore_ManifestData ManifestData;
+	ManifestData.VersionInfo.Initialize();
+	
 	const bool bSuccess = FileSandboxCore::SaveManifest(ManifestData, *InRootDirectory) && SaveMetaData(InMetaData, InRootDirectory);
 	return bSuccess ? MakeUnique<FSandboxInstance>(MoveTemp(InRootDirectory), InInitArgs, ManifestData, InMetaData) : nullptr;
 }
diff --git a/Engine/Plugins/Developer/Sandbox/FileSandbox/Source/FileSandboxCore/Public/Data/VersionInfo.h b/Engine/Plugins/Developer/Sandbox/FileSandbox/Source/FileSandboxCore/Public/Data/VersionInfo.h
new file mode 100644
--- /dev/null
+++ b/Engine/Plugins/Developer/Sandbox/FileSandbox/Source/FileSandboxCore/Public/Data/VersionInfo.h
@@ -0,0 +1,115 @@
+﻿// Copyright Epic Games, Inc. All Rights Reserved.
+
+#pragma once
+
+#include "VersionInfo.generated.h"
+
+class FEngineVersion;
+struct FCustomVersion;
+
+/** Holds file version information */
+USTRUCT()
+struct FFileSandboxCore_FileVersionInfo
+{
+	GENERATED_BODY()
+
+	/* UE5 File version */
+	UPROPERTY()
+	int32 FileVersionUE5 = 0;
+	
+	/* Licensee file version */
+	UPROPERTY()
+	int32 FileVersionLicensee = 0;
+
+	/** Initialize this version info from the compiled in data */
+	FILESANDBOXCORE_API void Initialize();
+	/** @return File version info as string */
+	FILESANDBOXCORE_API FString ToString() const;
+};
+
+/** Holds engine version information */
+USTRUCT()
+struct FFileSandboxCore_EngineVersionInfo
+{
+	GENERATED_BODY()
+
+	/** Major version number */
+	UPROPERTY()
+	uint16 Major = 0;
+
+	/** Minor version number */
+	UPROPERTY()
+	uint16 Minor = 0;
+
+	/** Patch version number */
+	UPROPERTY()
+	uint16 Patch = 0;
+
+	/** Changelist number. This is used to arbitrate when Major/Minor/Patch version numbers match */
+	UPROPERTY()
+	uint32 Changelist = 0;
+
+	/** Initialize this version info from the given version */
+	FILESANDBOXCORE_API void Initialize(const FEngineVersion& InVersion);
+
+	/** @return Engine version info as string, e.g. "5.7.1-48893048" */
+	FILESANDBOXCORE_API FString ToString(bool bIncludeChangelist = true) const;
+};
+
+/** Holds custom version information */
+USTRUCT()
+struct FFileSandboxCore_CustomVersionInfo
+{
+	GENERATED_BODY()
+
+	friend FArchive& operator<<(FArchive& Archive, FFileSandboxCore_CustomVersionInfo& VersionInfo)
+	{
+		Archive << VersionInfo.FriendlyName;
+		Archive << VersionInfo.Key;
+		Archive << VersionInfo.Version;
+		return Archive;
+	}
+	
+	/** Friendly name of the version */
+	UPROPERTY()
+	FName FriendlyName;
+
+	/** Unique custom key */
+	UPROPERTY()
+	FGuid Key;
+
+	/** Custom version */
+	UPROPERTY()
+	int32 Version = 0;
+
+	/** Initialize this version info from the given version */
+	FILESANDBOXCORE_API void Initialize(const FCustomVersion& InVersion);
+};
+
+/** Holds version information for a session */
+USTRUCT()
+struct FFileSandboxCore_VersionInfo
+{
+	GENERATED_BODY()
+	
+	/** File version info */
+	UPROPERTY()
+	FFileSandboxCore_FileVersionInfo FileVersion;
+
+	/** Engine version info */
+	UPROPERTY()
+	FFileSandboxCore_EngineVersionInfo EngineVersion;
+
+	/** Custom version info */
+	UPROPERTY()
+	TArray<FFileSandboxCore_CustomVersionInfo> CustomVersions;
+
+	/** Initialize this version info from the compiled in data */
+	FILESANDBOXCORE_API void Initialize();
+	/** @return Whether any data is set. */
+	FILESANDBOXCORE_API bool IsInitialized() const;
+	
+	/** @eturn File and version info as string. Does not include custom versions. */
+	FILESANDBOXCORE_API FString ToString() const;
+
+};
