File Sandbox Core: Create API for getting metadata
- Last Modified
- Version info

#Synced-CL 47637010

diff --git a/Engine/Plugins/Developer/Sandbox/FileSandbox/Source/FileSandboxCore/Private/Repository/NaiveSandboxRepository.cpp b/Engine/Plugins/Developer/Sandbox/FileSandbox/Source/FileSandboxCore/Private/Repository/NaiveSandboxRepository.cpp
--- a/Engine/Plugins/Developer/Sandbox/FileSandbox/Source/FileSandboxCore/Private/Repository/NaiveSandboxRepository.cpp
+++ b/Engine/Plugins/Developer/Sandbox/FileSandbox/Source/FileSandboxCore/Private/Repository/NaiveSandboxRepository.cpp
@@ -2,9 +2,11 @@
 
 #include "NaiveSandboxRepository.h"
 
+#include "Data/VersionInfo.h"
 #include "ISandboxManager.h"
 #include "Types/EBreakBehavior.h"
 #include "Types/RepositoryChangedEvent.h"
+#include "Types/SandboxMetaInfo.h"
 #include "Utils/SandboxDirectoryUtils.h"
 
 namespace UE::FileSandboxCore
@@ -21,12 +23,12 @@ FNaiveSandboxRepository::~FNaiveSandboxRepository()
 	Manager.OnPostSandboxStartup().RemoveAll(this);
 }
 
-void FNaiveSandboxRepository::ForEachSandbox(TFunctionRef<EBreakBehavior(const FString& InRootPath, const FFileSandboxCore_SandboxMetaData&)> InProcessSandboxes)
+void FNaiveSandboxRepository::ForEachSandbox(TFunctionRef<EBreakBehavior(const FString& InRootPath, const FSandboxMetaInfo&)> InProcessSandboxes)
 {
 	FileSandboxCore::ForEachSandbox([this, &InProcessSandboxes](const FString& InRootPath)
 	{
 		EBreakBehavior Behavior = EBreakBehavior::Continue;
-		ReadMetaData(InRootPath, [&InRootPath, &InProcessSandboxes, &Behavior](const FFileSandboxCore_SandboxMetaData& InMetaData)
+		ReadMetaData(InRootPath, [&InRootPath, &InProcessSandboxes, &Behavior](const FSandboxMetaInfo& InMetaData)
 		{
 			Behavior = InProcessSandboxes(InRootPath, InMetaData);
 		});
@@ -42,11 +44,14 @@ int32 FNaiveSandboxRepository::NumSandboxes() const
 	return Num;
 }
 
-bool FNaiveSandboxRepository::ReadMetaData(const FString& InRootPath, TFunctionRef<void(const FFileSandboxCore_SandboxMetaData&)> InProcessMetadata)
+bool FNaiveSandboxRepository::ReadMetaData(const FString& InRootPath, TFunctionRef<void(const FSandboxMetaInfo&)> InProcessMetadata)
 {
-	if (const TOptional<FFileSandboxCore_SandboxMetaData> MetaData = FileSandboxCore::LoadMetaData(InRootPath))
+	const TOptional<FFileSandboxCore_SandboxMetaData> MetaData = LoadMetaData(InRootPath);
+	const FDateTime DateTime = GetLastModified(InRootPath);
+	const FFileSandboxCore_VersionInfo VersionInfo = LoadVersionInfo(InRootPath);
+	if (MetaData && DateTime != FDateTime::MinValue() && VersionInfo.IsInitialized())
 	{
-		InProcessMetadata(*MetaData);
+		InProcessMetadata(FSandboxMetaInfo(*MetaData, DateTime, VersionInfo));
 		return true;
 	}
 	return false;
diff --git a/Engine/Plugins/Developer/Sandbox/FileSandbox/Source/FileSandboxCore/Private/Repository/NaiveSandboxRepository.h b/Engine/Plugins/Developer/Sandbox/FileSandbox/Source/FileSandboxCore/Private/Repository/NaiveSandboxRepository.h
--- a/Engine/Plugins/Developer/Sandbox/FileSandbox/Source/FileSandboxCore/Private/Repository/NaiveSandboxRepository.h
+++ b/Engine/Plugins/Developer/Sandbox/FileSandbox/Source/FileSandboxCore/Private/Repository/NaiveSandboxRepository.h
@@ -26,9 +26,9 @@ public:
 	~FNaiveSandboxRepository();
 
 	//~ Begin ISandboxRepository Interface
-	virtual void ForEachSandbox(TFunctionRef<EBreakBehavior(const FString& InRootPath, const FFileSandboxCore_SandboxMetaData&)> InProcessSandboxes) override;
+	virtual void ForEachSandbox(TFunctionRef<EBreakBehavior(const FString& InRootPath, const FSandboxMetaInfo&)> InProcessSandboxes) override;
 	virtual int32 NumSandboxes() const override;
-	virtual bool ReadMetaData(const FString& InRootPath, TFunctionRef<void(const FFileSandboxCore_SandboxMetaData&)> InProcessMetadata) override;
+	virtual bool ReadMetaData(const FString& InRootPath, TFunctionRef<void(const FSandboxMetaInfo&)> InProcessMetadata) override;
 	virtual FOnSandboxesChanged& OnSandboxesChanged() override { return OnSandboxesChangedDelegate; }
 	virtual FOnSandboxMetaDataChanged& OnSandboxMetaDataChanged() override { return OnSandboxMetaDataChangedDelegate; }
 	//~ End ISandboxRepository Interface
diff --git a/Engine/Plugins/Developer/Sandbox/FileSandbox/Source/FileSandboxCore/Private/Repository/WatchedSandboxRepository.cpp b/Engine/Plugins/Developer/Sandbox/FileSandbox/Source/FileSandboxCore/Private/Repository/WatchedSandboxRepository.cpp
--- a/Engine/Plugins/Developer/Sandbox/FileSandbox/Source/FileSandboxCore/Private/Repository/WatchedSandboxRepository.cpp
+++ b/Engine/Plugins/Developer/Sandbox/FileSandbox/Source/FileSandboxCore/Private/Repository/WatchedSandboxRepository.cpp
@@ -55,23 +55,23 @@ FWatchedSandboxRepository::FWatchedSandboxRepository(FString InBaseDirectory)
 	UpdateWatchedDirectories();
 }
 
-void FWatchedSandboxRepository::ForEachSandbox(TFunctionRef<EBreakBehavior(const FString& InRootPath, const FFileSandboxCore_SandboxMetaData& MetaData)> InProcessSandboxes)
+void FWatchedSandboxRepository::ForEachSandbox(TFunctionRef<EBreakBehavior(const FString& InRootPath, const FSandboxMetaInfo& MetaData)> InProcessSandboxes)
 {
 	for (const FSandboxData& SandboxData : CachedSandboxData)
 	{
-		if (InProcessSandboxes(SandboxData.AbsPath, SandboxData.MetaData) == EBreakBehavior::Break)
+		if (InProcessSandboxes(SandboxData.AbsPath, SandboxData.CachedData) == EBreakBehavior::Break)
 		{
 			break;
 		}
 	}
 }
 
-bool FWatchedSandboxRepository::ReadMetaData(const FString& InRootPath, TFunctionRef<void(const FFileSandboxCore_SandboxMetaData&)> InProcessMetadata)
+bool FWatchedSandboxRepository::ReadMetaData(const FString& InRootPath, TFunctionRef<void(const FSandboxMetaInfo&)> InProcessMetadata)
 {
 	const int32 Index = IndexOf(InRootPath);
 	if (CachedSandboxData.IsValidIndex(Index))
 	{
-		InProcessMetadata(CachedSandboxData[Index].MetaData);
+		InProcessMetadata(CachedSandboxData[Index].CachedData);
 		return true;
 	}
 	return false;
@@ -97,20 +97,23 @@ void FWatchedSandboxRepository::AddSandboxIfExists(const FString& InDirectory)
 	
 	const FString AbsPath = FPaths::ConvertRelativePathToFull(InDirectory);
 	TOptional<FFileSandboxCore_SandboxMetaData> MetaData = LoadMetaData(AbsPath);
-	if (!MetaData)
+	FDateTime LastModified = GetLastModified(AbsPath);
+	FFileSandboxCore_VersionInfo VersionInfo = LoadVersionInfo(AbsPath);
+	if (!MetaData || LastModified == FDateTime::MinValue() || !VersionInfo.IsInitialized())
 	{
 		return;
 	}
 	
 	const int32 ExistingIndex = IndexOf(AbsPath);
+	FSandboxMetaInfo MetaInfo(MoveTemp(*MetaData), LastModified, MoveTemp(VersionInfo));
 	if (CachedSandboxData.IsValidIndex(ExistingIndex)) // This is just defensive code - this case should not occur.
 	{
-		CachedSandboxData[ExistingIndex].MetaData = MoveTemp(*MetaData);
+		CachedSandboxData[ExistingIndex].CachedData = MoveTemp(MetaInfo);
 		OnSandboxMetaDataChangedDelegate.Broadcast(AbsPath);
 	}
 	else
 	{
-		CachedSandboxData.Emplace(FSandboxData{ AbsPath, MoveTemp(*MetaData) });
+		CachedSandboxData.Emplace(FSandboxData{ AbsPath, MoveTemp(MetaInfo) });
 		OnSandboxesChanged().Broadcast(FRepositoryChangedEvent{ .AddedSandboxPaths = MakeArrayView(&AbsPath, 1) });
 	}
 }
@@ -144,6 +147,10 @@ void FWatchedSandboxRepository::OnSandboxDirectoryChanged(const TArray<FFileChan
 		AddSandboxIfExists(InSandboxPath);
 		return;
 	}
+	if (!bWasSandboxDirectoryBefore)
+	{
+		return;
+	}
 	
 	if (!bIsSandboxDirectory)
 	{
@@ -152,10 +159,17 @@ void FWatchedSandboxRepository::OnSandboxDirectoryChanged(const TArray<FFileChan
 		return;
 	}
 	
-	DetectMetaDataFileChange(Index, InChanges, InSandboxPath);
+	bool bChangedMetaData = false;
+	const auto SandboxNotRemoved = [this, &InSandboxPath, Index] { return CachedSandboxData.IsValidIndex(Index) && CachedSandboxData[Index].AbsPath == InSandboxPath; };
+	bChangedMetaData |= DetectMetaDataFileChange(Index, InChanges, InSandboxPath);
+	bChangedMetaData |= SandboxNotRemoved() && DetectManifestFileChange(Index, InChanges, InSandboxPath);
+	if (bChangedMetaData)
+	{
+		OnSandboxMetaDataChangedDelegate.Broadcast(InSandboxPath);
+	}
 }
 
-void FWatchedSandboxRepository::DetectMetaDataFileChange(int32 InSandboxIndex, const TArray<FFileChangeData>& InChanges, const FString& InSandboxPath)
+bool FWatchedSandboxRepository::DetectMetaDataFileChange(int32 InSandboxIndex, const TArray<FFileChangeData>& InChanges, const FString& InSandboxPath)
 {
 	const bool bContainsUnknownOrRescan = InChanges.ContainsByPredicate([](const FFileChangeData& Change)
 	{
@@ -168,17 +182,51 @@ void FWatchedSandboxRepository::DetectMetaDataFileChange(int32 InSandboxIndex, c
 	// If the directory watcher is unsure what happened (rescan / unknown), we'll treat the file as changed, i.e. parse it again.
 	if (!bContainsUnknownOrRescan && !InChanges.IsValidIndex(MetadataFileChangeIndex))
 	{
-		return;
+		return false;
 	}
 	
 	if (TOptional<FFileSandboxCore_SandboxMetaData> MetaData = LoadMetaData(InSandboxPath))
 	{
-		CachedSandboxData[InSandboxIndex].MetaData = MoveTemp(*MetaData);
+		CachedSandboxData[InSandboxIndex].CachedData.UserMetaData = MoveTemp(*MetaData);
+		return true;
+	}
+	else
+	{
+		// Failed to load metadata. We'll treat it as the metadata file being corrupted or no longer existing.
+		CachedSandboxData.RemoveAtSwap(InSandboxIndex);
+		OnSandboxesChangedDelegate.Broadcast(FRepositoryChangedEvent{ .RemovedSandboxPaths = MakeArrayView(&InSandboxPath, 1) });
+		return false;
+	}
+}
+
+bool FWatchedSandboxRepository::DetectManifestFileChange(int32 InSandboxIndex, const TArray<FFileChangeData>& InChanges, const FString& InSandboxPath)
+{
+	const bool bContainsUnknownOrRescan = InChanges.ContainsByPredicate([](const FFileChangeData& Change)
+	{
+		return Change.Action == FFileChangeData::FCA_RescanRequired || Change.Action == FFileChangeData::FCA_Unknown;
+	});
+	const int32 ManifestFileChangeIndex = InChanges.IndexOfByPredicate([](const FFileChangeData& Change)
+	{
+		return FPaths::GetCleanFilename(Change.Filename) == GetManifestFileName();
+	});
+	
+	if (!bContainsUnknownOrRescan && !InChanges.IsValidIndex(ManifestFileChangeIndex))
+	{
+		return false;
+	}
+	
+	const FDateTime LastModified = GetLastModified(InSandboxPath);
+	if (LastModified != FDateTime::MinValue())
+	{
+		CachedSandboxData[InSandboxIndex].CachedData.LastModified = LastModified;
+		return true;
 	}
 	else
 	{
+		// Failed to obtain info about timestamp. We'll treat it as the file no longer existing.
 		CachedSandboxData.RemoveAtSwap(InSandboxIndex);
 		OnSandboxesChangedDelegate.Broadcast(FRepositoryChangedEvent{ .RemovedSandboxPaths = MakeArrayView(&InSandboxPath, 1) });
+		return false;
 	}
 }
 
@@ -196,8 +244,11 @@ void FWatchedSandboxRepository::UpdateCacheFromChangedBaseDirectory(const TArray
 	}
 	
 	TSet<FString> AlreadyProcessed;
-	for (const FFileChangeData& Change : InChanges)
+	// Process in reverse order to prioritize the latest change and ignore earlier changes.
+	// E.g. if InChanges is { { Modified, "Foo"}, { Removed, "Foo"} }, we just care about the Removed event.
+	for (int32 FileIndex = InChanges.Num() - 1; FileIndex >= 0; --FileIndex)
 	{
+		const FFileChangeData& Change = InChanges[FileIndex];
 		const FString& Path = Change.Filename;
 		
 		// InChanges can contain 100s of FCA_Modified changes for the same file. Don't do double work.
@@ -211,9 +262,9 @@ void FWatchedSandboxRepository::UpdateCacheFromChangedBaseDirectory(const TArray
 		{
 		case FFileChangeData::FCA_Added: AddSandboxIfExists(Path); break;
 		case FFileChangeData::FCA_Removed: 
-			if (const int32 Index = IndexOf(Path))
+			if (const int32 SandboxIndex = IndexOf(Path); CachedSandboxData.IsValidIndex(SandboxIndex))
 			{
-				CachedSandboxData.RemoveAtSwap(Index);
+				CachedSandboxData.RemoveAtSwap(SandboxIndex);
 				OnSandboxesChanged().Broadcast(FRepositoryChangedEvent{ .RemovedSandboxPaths = MakeArrayView(&Path, 1) });
 			}
 			break;
diff --git a/Engine/Plugins/Developer/Sandbox/FileSandbox/Source/FileSandboxCore/Private/Repository/WatchedSandboxRepository.h b/Engine/Plugins/Developer/Sandbox/FileSandbox/Source/FileSandboxCore/Private/Repository/WatchedSandboxRepository.h
--- a/Engine/Plugins/Developer/Sandbox/FileSandbox/Source/FileSandboxCore/Private/Repository/WatchedSandboxRepository.h
+++ b/Engine/Plugins/Developer/Sandbox/FileSandbox/Source/FileSandboxCore/Private/Repository/WatchedSandboxRepository.h
@@ -3,9 +3,10 @@
 #pragma once
 
 #if UE_SANDBOX_WITH_DIRECTORY_WATCHER
-#include "ISandboxRepository.h"
 #include "Data/SandboxMetaData.h"
+#include "ISandboxRepository.h"
 #include "Templates/UnrealTemplate.h"
+#include "Types/SandboxMetaInfo.h"
 #include "Utils/ScopedWatchedDirectory.h"
 
 struct FFileChangeData;
@@ -20,9 +21,9 @@ public:
 	explicit FWatchedSandboxRepository(FString InBaseDirectory);
 
 	//~ Begin ISandboxRepository Interface
-	virtual void ForEachSandbox(TFunctionRef<EBreakBehavior(const FString& InRootPath, const FFileSandboxCore_SandboxMetaData& MetaData)> InProcessSandboxes) override;
+	virtual void ForEachSandbox(TFunctionRef<EBreakBehavior(const FString& InRootPath, const FSandboxMetaInfo& MetaData)> InProcessSandboxes) override;
 	virtual int32 NumSandboxes() const override { return CachedSandboxData.Num(); }
-	virtual bool ReadMetaData(const FString& InRootPath, TFunctionRef<void(const FFileSandboxCore_SandboxMetaData&)> InProcessMetadata) override;
+	virtual bool ReadMetaData(const FString& InRootPath, TFunctionRef<void(const FSandboxMetaInfo&)> InProcessMetadata) override;
 	virtual FOnSandboxesChanged& OnSandboxesChanged() override { return OnSandboxesChangedDelegate; }
 	virtual FOnSandboxMetaDataChanged& OnSandboxMetaDataChanged() override { return OnSandboxMetaDataChangedDelegate; }
 	//~ End ISandboxRepository Interface
@@ -36,8 +37,9 @@ private:
 	{
 		/** Absolute path to the sandbox root directory, i.e. the directory containing the metadata file. */
 		FString AbsPath;
-		/** Cached content of the sandbox metadata file (to avoid constant JSON parsing, file loading, etc.) */
-		FFileSandboxCore_SandboxMetaData MetaData;
+		
+		/** Cached info about the sandbox, such as of the sandbox metadata file (to avoid constant JSON parsing, file loading, etc.) */ 
+		FSandboxMetaInfo CachedData;
 	};
 	/** Cached sandbox data. Updated whenever info changes on disk. */
 	TArray<FSandboxData> CachedSandboxData;
@@ -72,8 +74,10 @@ private:
 	
 	/** Invoked when a change is made in any of the sandbox directories. Used to update metadata. */
 	void OnSandboxDirectoryChanged(const TArray<FFileChangeData>& InChanges, FString InSandboxPath);
-	/** Checks whether a change has been reported on the metadata file and invokes OnSandboxMetaDataChangedDelegate, if so. */
-	void DetectMetaDataFileChange(int32 InSandboxIndex, const TArray<FFileChangeData>& InChanges, const FString& InSandboxPath);
+	/** Checks whether a change has been reported on the metadata file. */
+	bool DetectMetaDataFileChange(int32 InSandboxIndex, const TArray<FFileChangeData>& InChanges, const FString& InSandboxPath);
+	/** Checks whether a change has been reported for the manifest file. */
+	bool DetectManifestFileChange(int32 InSandboxIndex, const TArray<FFileChangeData>& InChanges, const FString& InSandboxPath);
 };	
 }
 
diff --git a/Engine/Plugins/Developer/Sandbox/FileSandbox/Source/FileSandboxCore/Private/Utils/SandboxDirectoryUtils.cpp b/Engine/Plugins/Developer/Sandbox/FileSandbox/Source/FileSandboxCore/Private/Utils/SandboxDirectoryUtils.cpp
--- a/Engine/Plugins/Developer/Sandbox/FileSandbox/Source/FileSandboxCore/Private/Utils/SandboxDirectoryUtils.cpp
+++ b/Engine/Plugins/Developer/Sandbox/FileSandbox/Source/FileSandboxCore/Private/Utils/SandboxDirectoryUtils.cpp
@@ -5,7 +5,9 @@
 #include "FileSandboxUtils.h"
 #include "JsonObjectConverter.h"
 #include "LogFileSandbox.h"
+#include "Data/ManifestData.h"
 #include "Data/SandboxMetaData.h"
+#include "Data/VersionInfo.h"
 #include "HAL/PlatformFileManager.h"
 #include "Misc/FileHelper.h"
 #include "Misc/Paths.h"
@@ -76,6 +78,19 @@ TOptional<FFileSandboxCore_SandboxMetaData> LoadMetaData(const FString& InSandbo
 	return {};
 }
 
+FDateTime GetLastModified(const FString& InSandboxDirectory)
+{
+	IPlatformFile& PlatformFile = FPlatformFileManager::Get().GetPlatformFile();
+	const FString ManifestFile = FPaths::Combine(InSandboxDirectory, GetManifestFileName());
+	return PlatformFile.GetTimeStampLocal(*ManifestFile);
+}
+
+FFileSandboxCore_VersionInfo LoadVersionInfo(const FString& InSandboxDirectory)
+{
+	const TOptional<FFileSandboxCore_ManifestData> Manifest = LoadManifest(InSandboxDirectory);
+	return Manifest ? Manifest->VersionInfo : FFileSandboxCore_VersionInfo();
+}
+
 TOptional<FString> FindSandboxByName(
 	const FString& InSandboxName,
 	const FString& InBaseDirectory
diff --git a/Engine/Plugins/Developer/Sandbox/FileSandbox/Source/FileSandboxCore/Public/ISandboxRepository.h b/Engine/Plugins/Developer/Sandbox/FileSandbox/Source/FileSandboxCore/Public/ISandboxRepository.h
--- a/Engine/Plugins/Developer/Sandbox/FileSandbox/Source/FileSandboxCore/Public/ISandboxRepository.h
+++ b/Engine/Plugins/Developer/Sandbox/FileSandbox/Source/FileSandboxCore/Public/ISandboxRepository.h
@@ -2,15 +2,16 @@
 
 #pragma once
 
+#include "Data/VersionInfo.h"
 #include "Delegates/Delegate.h"
 #include "HAL/Platform.h"
-
-struct FFileSandboxCore_SandboxMetaData;
+#include "Misc/DateTime.h"
 
 namespace UE::FileSandboxCore
 {
 enum class EBreakBehavior : uint8;
 struct FRepositoryChangedEvent;
+struct FSandboxMetaInfo;
 	
 DECLARE_MULTICAST_DELEGATE_OneParam(FOnSandboxesChanged, const FRepositoryChangedEvent&);
 DECLARE_MULTICAST_DELEGATE_OneParam(FOnSandboxMetaDataChanged, const FString& /* InRootPath */);
@@ -30,7 +31,7 @@ public:
 	 * Enumerates all known sandboxes. 
 	 * InRootPath will be the absolute path to the root sandbox directory, i.e. the one containing the metadata file.
 	 */
-	virtual void ForEachSandbox(TFunctionRef<EBreakBehavior(const FString& InRootPath, const FFileSandboxCore_SandboxMetaData& MetaData)> InProcessSandboxes) = 0;
+	virtual void ForEachSandbox(TFunctionRef<EBreakBehavior(const FString& InRootPath, const FSandboxMetaInfo& MetaData)> InProcessSandboxes) = 0;
 	
 	/** @return Number of distinct sandboxes, i.e. number of times ForEachSandbox would invoke the callback. */
 	virtual int32 NumSandboxes() const = 0;
@@ -41,7 +42,7 @@ public:
 	 * @param InProcessMetadata Callback invoked if the metadata exists
 	 * @return Whether the callback was invoked
 	 */
-	virtual bool ReadMetaData(const FString& InRootPath, TFunctionRef<void(const FFileSandboxCore_SandboxMetaData&)> InProcessMetadata) = 0;
+	virtual bool ReadMetaData(const FString& InRootPath, TFunctionRef<void(const FSandboxMetaInfo&)> InProcessMetadata) = 0;
 	
 	
 	/** @return Invoked when the known sandboxes changes (removed, discovered, etc.), on game thread. */
diff --git a/Engine/Plugins/Developer/Sandbox/FileSandbox/Source/FileSandboxCore/Public/Types/SandboxMetaInfo.h b/Engine/Plugins/Developer/Sandbox/FileSandbox/Source/FileSandboxCore/Public/Types/SandboxMetaInfo.h
new file mode 100644
--- /dev/null
+++ b/Engine/Plugins/Developer/Sandbox/FileSandbox/Source/FileSandboxCore/Public/Types/SandboxMetaInfo.h
@@ -0,0 +1,27 @@
+﻿// Copyright Epic Games, Inc. All Rights Reserved.
+
+#pragma once
+
+#include "Data/SandboxMetaData.h"
+#include "Data/VersionInfo.h"
+#include "Misc/DateTime.h"
+
+namespace UE::FileSandboxCore
+{
+struct FSandboxMetaInfo
+{
+	/** User / API-user provided meta data. */
+	FFileSandboxCore_SandboxMetaData UserMetaData;
+	
+	/** Local time that the sandbox was last modified. */
+	FDateTime LastModified;
+	/** Version that the sandbox was created in. */
+	FFileSandboxCore_VersionInfo VersionInfo;
+
+	explicit FSandboxMetaInfo(
+		FFileSandboxCore_SandboxMetaData InMetaData, FDateTime InLastModified, FFileSandboxCore_VersionInfo InVersionInfo
+		)
+		: UserMetaData(MoveTemp(InMetaData)), LastModified(InLastModified), VersionInfo(MoveTemp(InVersionInfo))
+	{}
+};
+}
diff --git a/Engine/Plugins/Developer/Sandbox/FileSandbox/Source/FileSandboxCore/Public/Utils/SandboxDirectoryUtils.h b/Engine/Plugins/Developer/Sandbox/FileSandbox/Source/FileSandboxCore/Public/Utils/SandboxDirectoryUtils.h
--- a/Engine/Plugins/Developer/Sandbox/FileSandbox/Source/FileSandboxCore/Public/Utils/SandboxDirectoryUtils.h
+++ b/Engine/Plugins/Developer/Sandbox/FileSandbox/Source/FileSandboxCore/Public/Utils/SandboxDirectoryUtils.h
@@ -9,6 +9,7 @@
 #define UE_API FILESANDBOXCORE_API
 
 class FText;
+struct FFileSandboxCore_VersionInfo;
 
 struct FFileSandboxCore_SandboxMetaData;
 
@@ -55,7 +56,13 @@ UE_API bool SaveMetaData(
  */
 UE_API TOptional<FFileSandboxCore_SandboxMetaData> LoadMetaData(const FString& InSandboxDirectory);
 
-	
+
+/** @return When the given sandbox was last modified. Returns FDateTime::MinValue() on failure. */
+UE_API FDateTime GetLastModified(const FString& InSandboxDirectory);
+
+/** @return Version info for the sandbox saved in this directory. This loads the data from the manifest file. IsInitialized returns false if this fails. */
+UE_API FFileSandboxCore_VersionInfo LoadVersionInfo(const FString& InSandboxDirectory);
+
 	
 /**
  * Searches for sandboxes with the given name. The search only considers the top-level directories.
