File Sandbox Core: Fixed bugs in persist logic and cleaned up the resulting logic.

Changes:
- Persisting now reports ESandboxFileChange. This type will be reused in a follow up API change to introspect changes made by the sandbox.
- Error handling has been simplified via IPersistErrorHandler. The intention is to make this API public in the future.
- Fixed some SC persist errors that would happen if you touched a file that was previously tracked by source control:
- Delete in SC before sandbox > Add file during sandbox > Persisting should checkout
- Add in SC before sandbox > Remove during sandbox > Persisting reverts the file
- Add in SC before sandbox > Edit during sandbox > Persisting retains the file as added
- Checkout file before sandbox > Remove during sandbox > Persisting marks for delete

#Synced-CL 47637010

diff --git a/Engine/Plugins/Developer/Sandbox/FileSandbox/Source/FileSandboxCore/Private/Sandbox/Persist/AggregatePersistErrorHandler.h b/Engine/Plugins/Developer/Sandbox/FileSandbox/Source/FileSandboxCore/Private/Sandbox/Persist/AggregatePersistErrorHandler.h
new file mode 100644
--- /dev/null
+++ b/Engine/Plugins/Developer/Sandbox/FileSandbox/Source/FileSandboxCore/Private/Sandbox/Persist/AggregatePersistErrorHandler.h
@@ -0,0 +1,60 @@
+﻿// Copyright Epic Games, Inc. All Rights Reserved.
+
+#pragma once 
+
+#include "Algo/ForEach.h"
+#include "Containers/Array.h"
+#include "IPersistErrorHandler.h"
+#include "Misc/NotNull.h"
+
+namespace UE::FileSandboxCore
+{
+/** Forwards the errors to other error handlers. */
+class FAggregatePersistErrorHandler : public IPersistErrorHandler
+{
+public:
+	
+	explicit FAggregatePersistErrorHandler(TArray<TNotNull<IPersistErrorHandler*>> InErrorHandlers) : ErrorHandlers(MoveTemp(InErrorHandlers)) {}
+	explicit FAggregatePersistErrorHandler(IPersistErrorHandler& InErrorHandler) : ErrorHandlers({ &InErrorHandler }) {}
+	
+	//~ Begin IPersistErrorHandler Interface
+	virtual void HandleError_CheckoutNotAllowed(const TCHAR* InFilename) override
+	{
+		Algo::ForEach(ErrorHandlers, [&](TNotNull<IPersistErrorHandler*> Handler){ Handler->HandleError_CheckoutNotAllowed(InFilename); });
+	}
+	virtual void HandleError_Checkout(const TCHAR* InFilename) override
+	{
+		Algo::ForEach(ErrorHandlers, [&](TNotNull<IPersistErrorHandler*> Handler){ Handler->HandleError_Checkout(InFilename); });
+	}
+	virtual void HandleError_Revert(const TCHAR* InFilename) override
+	{
+		Algo::ForEach(ErrorHandlers, [&](TNotNull<IPersistErrorHandler*> Handler){ Handler->HandleError_Revert(InFilename); });
+	}
+	virtual void HandleError_MarkForAdd(const TCHAR* InFilename) override
+	{
+		Algo::ForEach(ErrorHandlers, [&](TNotNull<IPersistErrorHandler*> Handler){ Handler->HandleError_MarkForAdd(InFilename); });
+	}
+	virtual void HandleError_DeleteSCC(const TCHAR* InFilename) override
+	{
+		Algo::ForEach(ErrorHandlers, [&](TNotNull<IPersistErrorHandler*> Handler){ Handler->HandleError_DeleteSCC(InFilename); });
+	}
+	virtual void HandleError_MakeWritable(const TCHAR* InFilename) override
+	{
+		Algo::ForEach(ErrorHandlers, [&](TNotNull<IPersistErrorHandler*> Handler){ Handler->HandleError_MakeWritable(InFilename); });
+	}
+	virtual void HandleError_MoveFile(const TCHAR* InToFilename, const TCHAR* InFromFilename) override
+	{
+		Algo::ForEach(ErrorHandlers, [&](TNotNull<IPersistErrorHandler*> Handler){ Handler->HandleError_MoveFile(InToFilename, InFromFilename); });
+	}
+	virtual void HandleError_DeleteFile(const TCHAR* InFilename) override
+	{
+		Algo::ForEach(ErrorHandlers, [&](TNotNull<IPersistErrorHandler*> Handler){ Handler->HandleError_DeleteFile(InFilename); });
+	}
+	//~ End IPersistErrorHandler Interface
+	
+private:
+	
+	/** The error handlers to forward the errors to. */
+	TArray<TNotNull<IPersistErrorHandler*>> ErrorHandlers;
+};
+}
diff --git a/Engine/Plugins/Developer/Sandbox/FileSandbox/Source/FileSandboxCore/Private/Sandbox/Persist/AppendTextErrorHandler.h b/Engine/Plugins/Developer/Sandbox/FileSandbox/Source/FileSandboxCore/Private/Sandbox/Persist/AppendTextErrorHandler.h
new file mode 100644
--- /dev/null
+++ b/Engine/Plugins/Developer/Sandbox/FileSandbox/Source/FileSandboxCore/Private/Sandbox/Persist/AppendTextErrorHandler.h
@@ -0,0 +1,83 @@
+﻿// Copyright Epic Games, Inc. All Rights Reserved.
+
+#pragma once 
+
+#include "IPersistErrorHandler.h"
+#include "LogFileSandbox.h"
+#include "Internationalization/Internationalization.h"
+#include "Internationalization/Text.h"
+#include "Logging/LogMacros.h"
+
+#define LOCTEXT_NAMESPACE "FAppendTextPersistErrorHandler"
+
+namespace UE::FileSandboxCore
+{
+class FAppendTextPersistErrorHandler : public IPersistErrorHandler
+{
+public:
+	
+	TArray<FText> AllErrors;
+	
+	//~ Begin IPersistErrorHandler Interface
+	virtual void HandleError_CheckoutNotAllowed(const TCHAR* InFilename) override
+	{
+		AllErrors.Emplace(FText::Format(
+			LOCTEXT("CheckoutNotAllowed", "Checkout of file '{0}' from revision control was not allowed when persisting sandbox state!"),
+			FText::AsCultureInvariant(InFilename)
+			));
+	}
+	virtual void HandleError_Checkout(const TCHAR* InFilename) override
+	{
+		AllErrors.Emplace(FText::Format(
+			LOCTEXT("Checkout", "Failed to checkout file '{0}' from revision control when persisting sandbox state!"),
+			FText::AsCultureInvariant(InFilename)
+			));
+	}
+	virtual void HandleError_Revert(const TCHAR* InFilename) override
+	{
+		AllErrors.Emplace(FText::Format(
+			LOCTEXT("Revert", "Failed to revert file '{0}' from revision control when persisting sandbox state!"),
+			FText::AsCultureInvariant(InFilename)
+			));
+	}
+	virtual void HandleError_MarkForAdd(const TCHAR* InFilename) override
+	{
+		AllErrors.Emplace(FText::Format(
+			LOCTEXT("MarkForAdd", "Failed to add file '{0}' to revision control when persisting sandbox state!"),
+			FText::AsCultureInvariant(InFilename)
+			));
+	}
+	virtual void HandleError_DeleteSCC(const TCHAR* InFilename) override
+	{
+		AllErrors.Emplace(FText::Format(
+			LOCTEXT("DeleteSCC", "Failed to delete file '{0}' from revision control when persisting sandbox state!"),
+			FText::AsCultureInvariant(InFilename)
+			));
+	}
+	virtual void HandleError_MakeWritable(const TCHAR* InFilename) override
+	{
+		AllErrors.Emplace(FText::Format(
+			LOCTEXT("MakeWritable", "Failed to make '{0}' writeable using file system when persisting sandbox state!"),
+			FText::AsCultureInvariant(InFilename)
+			));
+	}
+	virtual void HandleError_MoveFile(const TCHAR* InToFilename, const TCHAR* InFromFilename) override
+	{
+		AllErrors.Emplace(FText::Format(
+			LOCTEXT("MoveFile", "Failed to move file '{0}' to '{1}' using file system when persisting sandbox state!"),
+			FText::AsCultureInvariant(InFromFilename),
+			FText::AsCultureInvariant(InToFilename)
+			));
+	}
+	virtual void HandleError_DeleteFile(const TCHAR* InFilename) override
+	{
+		AllErrors.Emplace(FText::Format(
+			LOCTEXT("DeleteFile", "Failed to delete file '{0}' using file system when persisting sandbox state!"),
+			FText::AsCultureInvariant(InFilename)
+			));
+	}
+	//~ End IPersistErrorHandler Interface
+};
+}
+
+#undef LOCTEXT_NAMESPACE
\ No newline at end of file
diff --git a/Engine/Plugins/Developer/Sandbox/FileSandbox/Source/FileSandboxCore/Private/Sandbox/Persist/IPersistErrorHandler.h b/Engine/Plugins/Developer/Sandbox/FileSandbox/Source/FileSandboxCore/Private/Sandbox/Persist/IPersistErrorHandler.h
new file mode 100644
--- /dev/null
+++ b/Engine/Plugins/Developer/Sandbox/FileSandbox/Source/FileSandboxCore/Private/Sandbox/Persist/IPersistErrorHandler.h
@@ -0,0 +1,27 @@
+﻿// Copyright Epic Games, Inc. All Rights Reserved.
+
+#pragma once
+
+#include "Templates/UnrealTemplate.h"
+
+namespace UE::FileSandboxCore
+{
+/** Handles error that can occur during persist. */
+class IPersistErrorHandler
+{
+public:
+
+	virtual void HandleError_CheckoutNotAllowed(const TCHAR* InFilename) = 0;
+	virtual void HandleError_Checkout(const TCHAR* InFilename) = 0;
+	virtual void HandleError_Revert(const TCHAR* InFilename) = 0;
+	virtual void HandleError_MarkForAdd(const TCHAR* InFilename) = 0;
+	virtual void HandleError_DeleteSCC(const TCHAR* InFilename) = 0;
+	
+	virtual void HandleError_MakeWritable(const TCHAR* InFilename) = 0;
+	virtual void HandleError_MoveFile(const TCHAR* InToFilename, const TCHAR* InFromFilename) = 0;
+	virtual void HandleError_DeleteFile(const TCHAR* InFilename) = 0;
+
+	virtual ~IPersistErrorHandler() = default;
+};
+}
+
diff --git a/Engine/Plugins/Developer/Sandbox/FileSandbox/Source/FileSandboxCore/Private/Sandbox/Persist/LogPersistErrorHandler.h b/Engine/Plugins/Developer/Sandbox/FileSandbox/Source/FileSandboxCore/Private/Sandbox/Persist/LogPersistErrorHandler.h
new file mode 100644
--- /dev/null
+++ b/Engine/Plugins/Developer/Sandbox/FileSandbox/Source/FileSandboxCore/Private/Sandbox/Persist/LogPersistErrorHandler.h
@@ -0,0 +1,50 @@
+﻿// Copyright Epic Games, Inc. All Rights Reserved.
+
+#pragma once 
+
+#include "IPersistErrorHandler.h"
+#include "LogFileSandbox.h"
+#include "Logging/LogMacros.h"
+
+namespace UE::FileSandboxCore
+{
+class FLogPersistErrorHandler : public IPersistErrorHandler
+{
+public:
+	
+	//~ Begin IPersistErrorHandler Interface
+	virtual void HandleError_CheckoutNotAllowed(const TCHAR* InFilename) override
+	{
+		UE_LOG(LogFileSandbox, Warning, TEXT("Checkout of file '%s' from revision control was not allowed when persisting sandbox state!"), InFilename);
+	}
+	virtual void HandleError_Checkout(const TCHAR* InFilename) override
+	{
+		UE_LOG(LogFileSandbox, Warning, TEXT("Failed to checkout file '%s' from revision control when persisting sandbox state!"), InFilename);
+	}
+	virtual void HandleError_Revert(const TCHAR* InFilename) override
+	{
+		UE_LOG(LogFileSandbox, Warning, TEXT("Failed to revert file '%s' from revision control when persisting sandbox state!"), InFilename);
+	}
+	virtual void HandleError_MarkForAdd(const TCHAR* InFilename) override
+	{
+		UE_LOG(LogFileSandbox, Warning, TEXT("Failed to add file '%s' to revision control when persisting sandbox state!"), InFilename);
+	}
+	virtual void HandleError_DeleteSCC(const TCHAR* InFilename) override
+	{
+		UE_LOG(LogFileSandbox, Warning, TEXT("Failed to delete file '%s' from revision control when persisting sandbox state!"), InFilename);
+	}
+	virtual void HandleError_MakeWritable(const TCHAR* InFilename) override
+	{
+		UE_LOG(LogFileSandbox, Warning, TEXT("Failed to make '%s' writeable using file system when persisting sandbox state!"), InFilename);
+	}
+	virtual void HandleError_MoveFile(const TCHAR* InToFilename, const TCHAR* InFromFilename) override
+	{
+		UE_LOG(LogFileSandbox, Warning, TEXT("Failed to move file '%s' to '%s' using file system when persisting sandbox state!"), InFromFilename, InToFilename);
+	}
+	virtual void HandleError_DeleteFile(const TCHAR* InFilename) override
+	{
+		UE_LOG(LogFileSandbox, Warning, TEXT("Failed to delete file '%s' using file system when persisting sandbox state!"), InFilename);
+	}
+	//~ End IPersistErrorHandler Interface
+};
+}
diff --git a/Engine/Plugins/Developer/Sandbox/FileSandbox/Source/FileSandboxCore/Private/Sandbox/Persist/PersistUtils.cpp b/Engine/Plugins/Developer/Sandbox/FileSandbox/Source/FileSandboxCore/Private/Sandbox/Persist/PersistUtils.cpp
new file mode 100644
--- /dev/null
+++ b/Engine/Plugins/Developer/Sandbox/FileSandbox/Source/FileSandboxCore/Private/Sandbox/Persist/PersistUtils.cpp
@@ -0,0 +1,206 @@
+﻿// Copyright Epic Games, Inc. All Rights Reserved.
+
+#include "PersistUtils.h"
+
+#include "IPersistErrorHandler.h"
+#include "ISourceControlProvider.h"
+#include "ISourceControlState.h"
+#include "SourceControlOperations.h"
+#include "GenericPlatform/GenericPlatformFile.h"
+#include "Misc/Paths.h"
+#include "Utils/SandboxUtils.h"
+
+namespace UE::FileSandboxCore
+{
+namespace PersistDetail
+{
+static bool TryCheckout(
+	ISourceControlProvider& InSourceControlProvider, IPersistErrorHandler& InErrorHandler,
+	const TCHAR* InToFilename, const FSourceControlStateRef& ToFileSCCState
+	)
+{
+	// Handles e.g. exclusive checkout
+	const bool bIsCheckoutDisallowed = InSourceControlProvider.UsesCheckout() && !ToFileSCCState->IsCheckedOut() && !ToFileSCCState->CanCheckout();
+	if (bIsCheckoutDisallowed)
+	{
+		InErrorHandler.HandleError_CheckoutNotAllowed(InToFilename);
+		return false;
+	}
+		
+	if (ToFileSCCState->CanCheckout() && InSourceControlProvider.UsesCheckout())
+	{
+		TArray<FString> FilesToBeCheckedOut;
+		FilesToBeCheckedOut.Add(InToFilename);
+
+		if (InSourceControlProvider.Execute(ISourceControlOperation::Create<FCheckOut>(), FilesToBeCheckedOut) != ECommandResult::Succeeded)
+		{
+			InErrorHandler.HandleError_Checkout(InToFilename);
+			return false;
+		}
+	}
+	
+	return true;
+}
+}
+
+bool AddFileWithSCC(
+	IPlatformFile& InPlatformFile, ISourceControlProvider& InSourceControlProvider, IPersistErrorHandler& InErrorHandler,
+	const TCHAR* InToFilename, const TCHAR* InFromFilename,
+	EPersistFlags InFlags
+	)
+{
+	if (!MoveSandboxFileToNonSandbox(InPlatformFile, InErrorHandler, InToFilename, InFromFilename, InFlags))
+	{
+		return false;
+	}
+	
+	// If the file is new, add it to source control now
+	const FSourceControlStatePtr ToFileSCCState = InSourceControlProvider.IsEnabled() 
+		? InSourceControlProvider.GetState(InToFilename, EStateCacheUsage::ForceUpdate) : nullptr;
+	if (ToFileSCCState && !ToFileSCCState->IsSourceControlled())
+	{
+		const TArray<FString> FilesToBeAdded { InToFilename };
+		if (InSourceControlProvider.Execute(ISourceControlOperation::Create<FMarkForAdd>(), FilesToBeAdded) != ECommandResult::Succeeded)
+		{
+			InErrorHandler.HandleError_MarkForAdd(InToFilename);
+			return false;
+		}
+	}
+	// If the file was deleted before we entered the sandbox, but in the sandbox we added it, it means the file should be edited.
+	else if (ToFileSCCState && ToFileSCCState->IsDeleted())
+	{
+		if (InSourceControlProvider.Execute(ISourceControlOperation::Create<FRevert>(), InToFilename) != ECommandResult::Succeeded)
+		{
+			InErrorHandler.HandleError_Revert(InToFilename);
+			return false;
+		}
+		
+		return PersistDetail::TryCheckout(InSourceControlProvider, InErrorHandler, InToFilename, ToFileSCCState.ToSharedRef());
+	}
+	
+	return true;
+}
+
+bool EditFileWithSCC(
+	IPlatformFile& InPlatformFile, ISourceControlProvider& InSourceControlProvider, IPersistErrorHandler& InErrorHandler,
+	const TCHAR* InToFilename, const TCHAR* InFromFilename, 
+	bool bMarkWriteable, 
+	EPersistFlags InFlags
+	)
+{
+	const FSourceControlStatePtr ToFileSCCState = InSourceControlProvider.IsEnabled() 
+		? InSourceControlProvider.GetState(InToFilename, EStateCacheUsage::ForceUpdate) : nullptr;
+	
+	if (ToFileSCCState && ToFileSCCState->IsSourceControlled()) 
+	{
+		// If the file was marked for add before we edited it in sandbox, it should remain marked for add.
+		if (!ToFileSCCState->IsAdded() 
+			&& !PersistDetail::TryCheckout(InSourceControlProvider, InErrorHandler, InToFilename, ToFileSCCState.ToSharedRef()))
+		{
+			return false;
+		}
+	}
+	else if (bMarkWriteable)
+	{
+		if (InPlatformFile.FileExists(InToFilename) && !InPlatformFile.SetReadOnly(InToFilename, false))
+		{
+			InErrorHandler.HandleError_MakeWritable(InToFilename);
+			return false;
+		}
+	}
+	
+	return MoveSandboxFileToNonSandbox(InPlatformFile, InErrorHandler, InToFilename, InFromFilename, InFlags);
+}
+
+bool DeleteFileWithSCC(
+	IPlatformFile& InPlatformFile, ISourceControlProvider& InSourceControlProvider, IPersistErrorHandler& InErrorHandler, const TCHAR* InFilename
+	)
+{
+	// If this file maps to a package then we need to flush its linker so that we can remove the file from disk
+	FlushPackageFile(InFilename);
+
+	// Get the source control state of the file
+	const FSourceControlStatePtr FileSCCState = InSourceControlProvider.IsEnabled() 
+		? InSourceControlProvider.GetState(InFilename, EStateCacheUsage::ForceUpdate) : nullptr;
+	// Try and let source control remove the file first
+	if (FileSCCState && FileSCCState->IsSourceControlled())
+	{
+		const bool bAdded = FileSCCState->IsAdded();
+
+		const bool bCanRevert = bAdded || FileSCCState->IsCheckedOut(); 
+		if (bCanRevert && InSourceControlProvider.Execute(ISourceControlOperation::Create<FRevert>(), InFilename) != ECommandResult::Succeeded)
+		{
+			InErrorHandler.HandleError_Revert(InFilename);
+			return false;
+		}
+		
+		// If the file was marked for add, then revert is guaranteed to remove the file; it is invalid to call delete on that file again.
+		if (!bAdded && InSourceControlProvider.Execute(ISourceControlOperation::Create<FDelete>(), InFilename) != ECommandResult::Succeeded)
+		{
+			InErrorHandler.HandleError_DeleteSCC(InFilename);
+			return false;
+		}
+	}
+
+	// Delete file if it still exists
+	if (InPlatformFile.FileExists(InFilename) && !InPlatformFile.DeleteFile(InFilename))
+	{
+		InErrorHandler.HandleError_DeleteFile(InFilename);
+		return false;
+	}
+	return true;
+}
+
+namespace PersistDetail
+{
+static bool CopyFile(IPlatformFile& InPlatformFile, IPersistErrorHandler& InErrorHandler, const TCHAR* InToFilename, const TCHAR* InFromFilename)
+{
+	if (!InPlatformFile.CopyFile(InToFilename, InFromFilename, EPlatformFileRead::None, EPlatformFileWrite::AttemptDeleteAndCreate))
+	{
+		InErrorHandler.HandleError_MoveFile(InToFilename, InFromFilename);
+		return false;
+	}
+	return true;
+}
+
+static bool MoveFile(IPlatformFile& InPlatformFile, IPersistErrorHandler& InErrorHandler, const TCHAR* InToFilename, const TCHAR* InFromFilename)
+{
+	bool bSuccess = true;
+		
+	// MoveFile fails if there already is a target file
+	if (InPlatformFile.FileExists(InToFilename))
+	{
+		// DeleteFile fails if the file is read only
+		bSuccess &= !InPlatformFile.IsReadOnly(InToFilename) || InPlatformFile.SetReadOnly(InToFilename, false);
+		bSuccess &= InPlatformFile.DeleteFile(InToFilename);
+	}
+		
+	bSuccess &= InPlatformFile.MoveFile(InToFilename, InFromFilename);
+	
+	if (!bSuccess)
+	{
+		InErrorHandler.HandleError_MoveFile(InToFilename, InFromFilename);
+	}
+		
+	return bSuccess;
+}
+}
+
+bool MoveSandboxFileToNonSandbox(
+	IPlatformFile& InPlatformFile, IPersistErrorHandler& InErrorHandler,
+	const TCHAR* InToFilename, const TCHAR* InFromFilename,
+	EPersistFlags InFlags
+	)
+{
+	const FString ToFileDir = FPaths::GetPath(InToFilename);
+	if (!InPlatformFile.CreateDirectoryTree(*ToFileDir))
+	{
+		InErrorHandler.HandleError_MoveFile(InToFilename, InFromFilename);
+		return false;
+	}
+	
+	return EnumHasAnyFlags(InFlags, EPersistFlags::RetainChangedFiles)
+		? PersistDetail::CopyFile(InPlatformFile, InErrorHandler, InToFilename, InFromFilename)
+		: PersistDetail::MoveFile(InPlatformFile, InErrorHandler, InToFilename, InFromFilename);
+}
+}
diff --git a/Engine/Plugins/Developer/Sandbox/FileSandbox/Source/FileSandboxCore/Private/Sandbox/Persist/PersistUtils.h b/Engine/Plugins/Developer/Sandbox/FileSandbox/Source/FileSandboxCore/Private/Sandbox/Persist/PersistUtils.h
new file mode 100644
--- /dev/null
+++ b/Engine/Plugins/Developer/Sandbox/FileSandbox/Source/FileSandboxCore/Private/Sandbox/Persist/PersistUtils.h
@@ -0,0 +1,60 @@
+﻿// Copyright Epic Games, Inc. All Rights Reserved.
+
+#pragma once
+
+#include "HAL/Platform.h"
+#include "Misc/EnumClassFlags.h"
+
+class IPlatformFile;
+class ISourceControlProvider;
+
+namespace UE::FileSandboxCore
+{
+class FSandboxedPlatformFilePath;
+class IPersistErrorHandler;
+enum class ESandboxFileChange : uint8;
+struct FPersistArgs;
+
+enum class EPersistFlags : uint8
+{
+	None,
+	
+	/** 
+	 * If set, added & edited files are retained in the sandbox directory on persist, i.e. sandboxed files are copied to the non-sandbox files.
+	 * If unset, added & edited files are removed from the sandbox directory on persist, i.e. sandboxed files are moved to the non-sandbox files.
+	 */
+	RetainChangedFiles = 1 << 0
+};
+ENUM_CLASS_FLAGS(EPersistFlags);
+
+/** Persists a file that is missing from the non-sandbox but exists in the sandbox. Registers the changes with Source Control if enabled. */
+bool AddFileWithSCC(
+	IPlatformFile& InPlatformFile, ISourceControlProvider& InSourceControlProvider, IPersistErrorHandler& InErrorHandler,
+	const TCHAR* InToFilename, const TCHAR* InFromFilename, 
+	EPersistFlags InFlags = EPersistFlags::None
+	);
+/** Persists a file that exists in both the non-sandbox and sandbox. Registers the changes with Source Control if enabled. */
+bool EditFileWithSCC(
+	IPlatformFile& InPlatformFile, ISourceControlProvider& InSourceControlProvider, IPersistErrorHandler& InErrorHandler,
+	const TCHAR* InToFilename, const TCHAR* InFromFilename, 
+	bool bMarkWriteable, 
+	EPersistFlags InFlags = EPersistFlags::None
+	);
+/** Persists a file that exists in non-sandbox but was removed in the sandbox. Registers the changes with Source Control if enabled. */
+bool DeleteFileWithSCC(
+	IPlatformFile& InPlatformFile, ISourceControlProvider& InSourceControlProvider, IPersistErrorHandler& InErrorHandler,
+	const TCHAR* InFilename
+	);
+
+/** 
+ * Moves InFromFilename to InToFileName. 
+ * 
+ * After successful operation, InToFileName has the file content that InFromFilename previously had. 
+ * Depending on the flags InFromFilename stays in the sandbox, and continues to be tracked as change, or is removed.
+ */
+bool MoveSandboxFileToNonSandbox(
+	IPlatformFile& PlatformFile, IPersistErrorHandler& InErrorHandler, 
+	const TCHAR* InToFilename, const TCHAR* InFromFilename, 
+	EPersistFlags InFlags = EPersistFlags::None
+	);
+}
diff --git a/Engine/Plugins/Developer/Sandbox/FileSandbox/Source/FileSandboxCore/Private/Sandbox/Platform/SandboxPlatformFile.cpp b/Engine/Plugins/Developer/Sandbox/FileSandbox/Source/FileSandboxCore/Private/Sandbox/Platform/SandboxPlatformFile.cpp
--- a/Engine/Plugins/Developer/Sandbox/FileSandbox/Source/FileSandboxCore/Private/Sandbox/Platform/SandboxPlatformFile.cpp
+++ b/Engine/Plugins/Developer/Sandbox/FileSandbox/Source/FileSandboxCore/Private/Sandbox/Platform/SandboxPlatformFile.cpp
@@ -20,6 +20,8 @@
 #include "RenderingThread.h"
 #include "SandboxMountPoint.h"
 #include "SourceControlOperations.h"
+#include "Sandbox/Persist/PersistUtils.h"
+#include "Types/SandboxFileChange.h"
 #include "Types/Sandbox/PersistArgs.h"
 #include "Types/Sandbox/PersistResult.h"
 #include "UObject/LinkerLoad.h"
@@ -672,147 +674,37 @@ FString FSandboxPlatformFile::ConvertToAbsolutePathForExternalAppForWrite(const
 	return LowerLevel->ConvertToAbsolutePathForExternalAppForWrite(*ResolvedPath.GetNonSandboxPath());
 }
 
-
-bool FSandboxPlatformFile::CopyFileWithSCC(
-	ISourceControlProvider& InSourceControlProvider, 
-	const TCHAR* InToFilename,
-	const TCHAR* InFromFilename,
-	const FPersistArgs& InParam,
-	FPersistResult& OutResult
+ESandboxFileChange FSandboxPlatformFile::PersistItem(
+	const FSandboxedPlatformFilePath& InFile,
+	ISourceControlProvider& InSourceControlProvider, const FPersistArgs& InParams, IPersistErrorHandler& InErrorHandler
 	)
 {
-	// If this file maps to a package then we need to flush its linker so that we can overwrite the file on disk
-	FlushPackageFile(InToFilename);
-
-	// Get the source control state of the destination file
-	FSourceControlStatePtr ToFileSCCState;
-	if (InSourceControlProvider.IsEnabled())
-	{
-		ToFileSCCState = InSourceControlProvider.GetState(InToFilename, EStateCacheUsage::ForceUpdate);
-	}
-
-	// We don't need to do anything with source control if the file is already checked-out or added
-	const bool bRequiresSCCAction = ToFileSCCState && !ToFileSCCState->IsCheckedOut() && !ToFileSCCState->IsAdded();
-
-	// If the file can be checked-out, do so now
-	if (bRequiresSCCAction && ToFileSCCState->IsSourceControlled())
-	{
-		if (ToFileSCCState->CanCheckout() && InSourceControlProvider.UsesCheckout())
-		{
-			TArray<FString> FilesToBeCheckedOut;
-			FilesToBeCheckedOut.Add(InToFilename);
-
-			if (InSourceControlProvider.Execute(ISourceControlOperation::Create<FCheckOut>(), FilesToBeCheckedOut) != ECommandResult::Succeeded)
-			{
-				FText Failure = FText::Format(
-					NSLOCTEXT("Concert", "CheckoutFileFailure", "Failed to check-out file '{0}' from revision control when persiting sandbox state!"),
-					FText::FromString(InToFilename));
-				UE_LOG(LogFileSandbox, Warning, TEXT("%s"),  *Failure.ToString());
-				OutResult.FailureReasons.Add(MoveTemp(Failure));
-				return false;
-			}
-		}
-		else
-		{
-			FText Failure = FText::Format(
-				NSLOCTEXT("Concert", "CanCheckoutFileFailure", "Can't check-out file '{0}' from revision control when persiting sandbox state!"),
-				FText::FromString(InToFilename));
-			UE_LOG(LogFileSandbox, Warning, TEXT("%s"), *Failure.ToString());
-			OutResult.FailureReasons.Add(MoveTemp(Failure));
-			return false;
-		}
-	}
-	else if (InParam.bShouldMakeWritableIfNoSourceControl)
-	{
-		if (LowerLevel->FileExists(InToFilename) && !LowerLevel->SetReadOnly(InToFilename, false))
-		{
-			FText Failure = FText::Format(
-				NSLOCTEXT("Concert", "CanMakeWritable", "Can't make file '{0}' writable when persiting sandbox state!"),
-				FText::FromString(InToFilename));
-			UE_LOG(LogFileSandbox, Warning, TEXT("%s"), *Failure.ToString());
-			OutResult.FailureReasons.Add(MoveTemp(Failure));
-			return false;
-		}
-	}
-
-	// Copy the on-disk sandbox file
-	FString ToFileDir = FPaths::GetPath(InToFilename);
-	if (!CreateDirectoryTree(*ToFileDir) || !LowerLevel->CopyFile(InToFilename, InFromFilename, EPlatformFileRead::None, EPlatformFileWrite::AttemptDeleteAndCreate))
-	{
-		FText Failure = FText::Format(
-			NSLOCTEXT("Concert", "CopyFileFailure", "Failed to copy file '{0}' (from '{1}') when persiting sandbox state!"),
-			FText::FromString(InToFilename), FText::FromString(InFromFilename));
-		UE_LOG(LogFileSandbox, Warning, TEXT("%s"), *Failure.ToString());
-		OutResult.FailureReasons.Add(MoveTemp(Failure));
-		return false;
-	}
-
-		// If the file is new, add it to source control now
-	if (bRequiresSCCAction && !ToFileSCCState->IsSourceControlled())
-	{
-		TArray<FString> FilesToBeAdded;
-		FilesToBeAdded.Add(InToFilename);
-
-		if (InSourceControlProvider.Execute(ISourceControlOperation::Create<FMarkForAdd>(), FilesToBeAdded) != ECommandResult::Succeeded)
-		{
-			FText Failure = FText::Format(
-				NSLOCTEXT("Concert", "AddFileFailure", "Failed to add file '{0}' to revision control when persiting sandbox state!"),
-				FText::FromString(InToFilename));
-			UE_LOG(LogFileSandbox, Warning, TEXT("%s"), *Failure.ToString());
-			OutResult.FailureReasons.Add(MoveTemp(Failure));
-			return false;
-		}
-	}
-	return true;
-}
-
-bool FSandboxPlatformFile::DeleteFileWithSCC(ISourceControlProvider& InSourceControlProvider, const TCHAR* InFilename, FPersistResult& OutResult)
-{
-	// If this file maps to a package then we need to flush its linker so that we can remove the file from disk
-	FlushPackageFile(InFilename);
-
-	// Get the source control state of the file
-	const FSourceControlStatePtr FileSCCState = InSourceControlProvider.IsEnabled() 
-		? InSourceControlProvider.GetState(InFilename, EStateCacheUsage::ForceUpdate) : nullptr;
-	// Try and let source control remove the file first
-	if (FileSCCState && FileSCCState->IsSourceControlled())
+	if (IsPathDeleted(InFile))
 	{
-		const bool bAdded = FileSCCState->IsAdded();
-
-		if (bAdded || FileSCCState->IsCheckedOut())
-		{
-			if (InSourceControlProvider.Execute(ISourceControlOperation::Create<FRevert>(), InFilename) != ECommandResult::Succeeded)
-			{
-				FText Failure = FText::Format(NSLOCTEXT("Concert", "RevertFileFailure", "Failed to revert file '{0}' to revision control when persisting sandbox state!"), FText::FromString(InFilename));
-				UE_LOG(LogFileSandbox, Warning, TEXT("%s"), *Failure.ToString());
-				OutResult.FailureReasons.Add(MoveTemp(Failure));
-				return false;
-			}
-		}
-		
-		// If the file was marked for add, then revert is guaranteed to remove the file; it is invalid to call delete on that file again.
-		if (!bAdded && InSourceControlProvider.Execute(ISourceControlOperation::Create<FDelete>(), InFilename) != ECommandResult::Succeeded)
-		{
-			FText Failure = FText::Format(NSLOCTEXT("Concert", "DeleteSCCFileFailure", "Failed to delete file '{0}' from revision control when persiting sandbox state!"), FText::FromString(InFilename));
-			UE_LOG(LogFileSandbox, Warning, TEXT("%s"), *Failure.ToString());
-			OutResult.FailureReasons.Add(MoveTemp(Failure));
-			return false;
-		}
+		return DeleteFileWithSCC(*LowerLevel, InSourceControlProvider, InErrorHandler, *InFile.GetNonSandboxPath())
+			? ESandboxFileChange::Removed
+			: ESandboxFileChange::None; 
 	}
-
-	// Delete file if it still exists
-	if (LowerLevel->FileExists(InFilename) && !LowerLevel->DeleteFile(InFilename))
+	
+	const TCHAR* InToFilename = *InFile.GetNonSandboxPath();
+	const bool bExistsInNonSandbox = LowerLevel->FileExists(InToFilename);
+	const bool bIsAddOp = !bExistsInNonSandbox;
+	if (bIsAddOp)
 	{
-		FText Failure = FText::Format(NSLOCTEXT("Concert", "DeleteFileFailure", "Failed to delete file '{0}' when persiting sandbox state!"), FText::FromString(InFilename));
-		UE_LOG(LogFileSandbox, Warning, TEXT("%s"), *Failure.ToString());
-		OutResult.FailureReasons.Add(MoveTemp(Failure));
-		return false;
+		return AddFileWithSCC(*LowerLevel, InSourceControlProvider, InErrorHandler, InToFilename, *InFile.GetSandboxPath())
+			? ESandboxFileChange::Added 
+			: ESandboxFileChange::None;
 	}
-	return true;
+	return EditFileWithSCC(*LowerLevel, InSourceControlProvider, InErrorHandler, InToFilename, *InFile.GetSandboxPath(), InParams.bShouldMakeWritableIfNoSourceControl)
+		? ESandboxFileChange::Edited 
+		: ESandboxFileChange::None;
 }
 
-FPersistResult FSandboxPlatformFile::PersistSandbox(ISourceControlProvider& InSourceControlProvider,
-	FPersistArgs InParams, TFunctionRef<void(const FSandboxedPlatformFilePath&, EPersistedFileType)> InOnPersistFile
+bool FSandboxPlatformFile::PersistSandbox(
+	ISourceControlProvider& InSourceControlProvider,
+	const FPersistArgs& InParams, 
+	TFunctionRef<void(const FSandboxedPlatformFilePath&, ESandboxFileChange)> InOnPersistFile,
+	IPersistErrorHandler& InErrorHandler
 	)
 {
 	FPersistResult OutResult;
@@ -822,21 +714,18 @@ FPersistResult FSandboxPlatformFile::PersistSandbox(ISourceControlProvider& InSo
 	for (const FString& File : InParams.Files)
 	{
 		const FSandboxedPlatformFilePath FilePath = ToSandboxPath(File, true);
-		const EPersistedFileType Type = IsPathDeleted(FilePath) ? EPersistedFileType::Deleted : EPersistedFileType::Moved; 
-		const bool bSuccess = Type == EPersistedFileType::Deleted
-			? DeleteFileWithSCC(InSourceControlProvider, *FilePath.GetNonSandboxPath(), OutResult)
-			: CopyFileWithSCC(InSourceControlProvider, *FilePath.GetNonSandboxPath(), *FilePath.GetSandboxPath(), InParams, OutResult);
-
+		const ESandboxFileChange ActionPerformed = PersistItem(FilePath, InSourceControlProvider, InParams, InErrorHandler);
+		const bool bSuccess = ActionPerformed != ESandboxFileChange::None;
+		
 		// if the file operation was successful, mark the file as persisted
 		if (bSuccess)
 		{
-			InOnPersistFile(FilePath, Type);
+			InOnPersistFile(FilePath, ActionPerformed);
 		}
 
 		bFullOperationSuccess &= bSuccess;
 	}
-	OutResult.PersistStatus = bFullOperationSuccess ? EPersistStatus::Success : EPersistStatus::Failure;
-	return OutResult;
+	return bFullOperationSuccess;
 }
 
 void FSandboxPlatformFile::DiscardSandbox(
diff --git a/Engine/Plugins/Developer/Sandbox/FileSandbox/Source/FileSandboxCore/Private/Sandbox/Platform/SandboxPlatformFile.h b/Engine/Plugins/Developer/Sandbox/FileSandbox/Source/FileSandboxCore/Private/Sandbox/Platform/SandboxPlatformFile.h
--- a/Engine/Plugins/Developer/Sandbox/FileSandbox/Source/FileSandboxCore/Private/Sandbox/Platform/SandboxPlatformFile.h
+++ b/Engine/Plugins/Developer/Sandbox/FileSandbox/Source/FileSandboxCore/Private/Sandbox/Platform/SandboxPlatformFile.h
@@ -17,11 +17,14 @@ class ISourceControlProvider;
 
 namespace UE::FileSandboxCore
 {
-	struct FFileChange;
-	struct FDiscardSandboxArgs;
-	struct FSandboxMountPoint;
-	struct FPersistResult;
+class IPersistErrorHandler;
+enum class EFileChangeAction : uint8;
+enum class ESandboxFileChange : uint8;
+struct FDiscardSandboxArgs;
+struct FFileChange;
 struct FPersistArgs;
+struct FPersistResult;
+struct FSandboxMountPoint;
 	
 DECLARE_MULTICAST_DELEGATE_TwoParams(FDeletionStateChanged, const FSandboxedPlatformFilePath&, bool /* bIsDeleted */);
 	
@@ -69,16 +72,20 @@ public:
 	virtual FString ConvertToAbsolutePathForExternalAppForWrite(const TCHAR* Filename) override;
 	//~ end IPlatformFile overrides
 
-	enum class EPersistedFileType : uint8 { Deleted, Moved };
 	/**
 	 *	Persist the file list from the sandbox state onto the real files Will mark files for which the operation was
 	 *	successful as persisted.
 	 *	@param InSourceControlProvider The source control provider to register the changes with
 	 *	@param InParams	Parameters controlling the persist options specified by the user.
 	 *	@param InOnPersistFile Invoked for each file successfully persisted.
-	 *	@return Returns a persist result object that contains result status.
+	 *	@param InErrorHandler Handles any errors that occur during persist
+	 *	@return Whether the persist operation succeeded completely. False if there were errors.
 	 */
-	FPersistResult PersistSandbox(ISourceControlProvider& InSourceControlProvider, FPersistArgs InParams, TFunctionRef<void(const FSandboxedPlatformFilePath&, EPersistedFileType)> InOnPersistFile);
+	bool PersistSandbox(
+		ISourceControlProvider& InSourceControlProvider, const FPersistArgs& InParams, 
+		TFunctionRef<void(const FSandboxedPlatformFilePath&, ESandboxFileChange)> InOnPersistFile, 
+		IPersistErrorHandler& InErrorHandler
+		);
 
 	/** 
 	 *	Discard the sandbox state
@@ -154,10 +161,11 @@ private:
 	void OnSandboxDirectoryChanged(const TArray<FFileChangeData>& FileChanges, FSandboxedPlatformFilePath MountPath);
 #endif
 
-	/** Copy file using source control (if enabled) */
-	bool CopyFileWithSCC(ISourceControlProvider& InSourceControlProvider, const TCHAR* InToFilename, const TCHAR* InFromFilename, const FPersistArgs& InParam, FPersistResult& OutResult);
-	/** Delete file using source control (if enabled) */
-	bool DeleteFileWithSCC(ISourceControlProvider& InSourceControlProvider, const TCHAR* InFilename, FPersistResult& OutResult);
+	/** Attempts to persist the file. */
+	ESandboxFileChange PersistItem(
+		const FSandboxedPlatformFilePath& InFile, ISourceControlProvider& InSourceControlProvider, 
+		const FPersistArgs& InParams, IPersistErrorHandler& InErrorHandler
+		);
 
 	FString GetSandboxRootPath() const;
 		
diff --git a/Engine/Plugins/Developer/Sandbox/FileSandbox/Source/FileSandboxCore/Private/Sandbox/SandboxInstance.cpp b/Engine/Plugins/Developer/Sandbox/FileSandbox/Source/FileSandboxCore/Private/Sandbox/SandboxInstance.cpp
--- a/Engine/Plugins/Developer/Sandbox/FileSandbox/Source/FileSandboxCore/Private/Sandbox/SandboxInstance.cpp
+++ b/Engine/Plugins/Developer/Sandbox/FileSandbox/Source/FileSandboxCore/Private/Sandbox/SandboxInstance.cpp
@@ -4,6 +4,10 @@
 
 #include "HAL/FileManager.h"
 #include "Misc/PackageName.h"
+#include "Persist/AggregatePersistErrorHandler.h"
+#include "Persist/AppendTextErrorHandler.h"
+#include "Persist/LogPersistErrorHandler.h"
+#include "Types/SandboxFileChange.h"
 #include "Types/Sandbox/PersistArgs.h"
 #include "Types/Manager/SandboxInitArgs.h"
 #include "Types/Sandbox/DiscardResult.h"
@@ -95,30 +99,41 @@ FSandboxInstance::~FSandboxInstance()
 FPersistResult FSandboxInstance::PersistSandbox(const FPersistArgs& InArgs)
 {
 	bool bManifestNeedsSave = false;
-	const FPersistResult Result = PlatformFile->PersistSandbox(
-		SourceControlSandbox.GetProxiedSourceControl(), InArgs,
-		[this, &bManifestNeedsSave](const FSandboxedPlatformFilePath& InPath, FSandboxPlatformFile::EPersistedFileType InPersistType)
+	const auto HandleFileOp = [this, &bManifestNeedsSave](const FSandboxedPlatformFilePath& InPath, ESandboxFileChange InPersistType)
+	{
+		switch (InPersistType)
 		{
-			switch (InPersistType)
-			{
-			case FSandboxPlatformFile::EPersistedFileType::Deleted:
-				bManifestNeedsSave |= ManifestContent.DeletedFiles.Remove(InPath.GetNonSandboxPath());
-				break;
-			case FSandboxPlatformFile::EPersistedFileType::Moved:
-				bManifestNeedsSave |= ManifestContent.ModifiedFiles.Remove(InPath.GetNonSandboxPath());
-				bManifestNeedsSave |= ManifestContent.AddedFiles.Remove(InPath.GetNonSandboxPath());
-				break;
-				
-			default: checkNoEntry(); break;
-			}
-		});
+		case ESandboxFileChange::Removed:
+			bManifestNeedsSave |= ManifestContent.DeletedFiles.Remove(InPath.GetNonSandboxPath());
+			break;
+		case ESandboxFileChange::Added:
+			bManifestNeedsSave |= ManifestContent.AddedFiles.Remove(InPath.GetNonSandboxPath());
+			break;
+		case ESandboxFileChange::Edited:
+			bManifestNeedsSave |= ManifestContent.ModifiedFiles.Remove(InPath.GetNonSandboxPath());
+			break;
+			
+		default: checkNoEntry(); break;
+		}
+	};
+	
+	FAppendTextPersistErrorHandler ErrorsAsReasons;
+	FLogPersistErrorHandler ErrorLogger;
+	FAggregatePersistErrorHandler Aggregate({ &ErrorsAsReasons, &ErrorLogger });
+	const bool bFullSuccess = PlatformFile->PersistSandbox(
+		SourceControlSandbox.GetProxiedSourceControl(), InArgs, HandleFileOp, Aggregate
+		);
 
 	if (bManifestNeedsSave)
 	{
 		SaveManifest();
 	}
 
-	return Result;
+	return FPersistResult
+	{ 
+		MoveTemp(ErrorsAsReasons.AllErrors),
+		bFullSuccess ? EPersistStatus::Success : EPersistStatus::Failure 
+	};
 }
 
 FDiscardResult FSandboxInstance::DiscardSandbox(const FDiscardArgs& InArgs)
diff --git a/Engine/Plugins/Developer/Sandbox/FileSandbox/Source/FileSandboxCore/Public/Types/SandboxFileChange.h b/Engine/Plugins/Developer/Sandbox/FileSandbox/Source/FileSandboxCore/Public/Types/SandboxFileChange.h
new file mode 100644
--- /dev/null
+++ b/Engine/Plugins/Developer/Sandbox/FileSandbox/Source/FileSandboxCore/Public/Types/SandboxFileChange.h
@@ -0,0 +1,24 @@
+﻿// Copyright Epic Games, Inc. All Rights Reserved.
+
+#pragma once
+
+#include "HAL/Platform.h"
+
+namespace UE::FileSandboxCore
+{
+/** Types of changes the sandbox system may make to a file. */
+enum class ESandboxFileChange : uint8
+{
+	/** No changes were made to the file. */
+	None,
+	
+	/** The file was added. */
+	Added,
+	/** The file was removed. */
+	Removed,
+	/** The file was edited. */
+	Edited
+	
+	// TODO UE-356393: Add support for moved files & moved-and-edited files.
+};
+}
\ No newline at end of file
