File Sandbox Core: Add API to disallow leaving the sandbox.

#Synced-CL 47637010

diff --git a/Engine/Plugins/Developer/Sandbox/FileSandbox/Source/FileSandboxCore/Private/ISandboxManager.cpp b/Engine/Plugins/Developer/Sandbox/FileSandbox/Source/FileSandboxCore/Private/ISandboxManager.cpp
--- a/Engine/Plugins/Developer/Sandbox/FileSandbox/Source/FileSandboxCore/Private/ISandboxManager.cpp
+++ b/Engine/Plugins/Developer/Sandbox/FileSandbox/Source/FileSandboxCore/Private/ISandboxManager.cpp
@@ -2,6 +2,7 @@
 
 #include "ISandboxManager.h"
 
+#include "Interface/ISandboxLock.h"
 #include "Types/Manager/DeleteSandboxByDirectoryArgs.h"
 #include "Types/Manager/DeleteSandboxByNameArgs.h"
 #include "Types/Manager/DeleteSandboxResult.h"
@@ -36,4 +37,15 @@ FDeleteSandboxResult ISandboxManager::DeleteNamedSandbox(const FDeleteSandboxByN
 	
 	return DeleteSandbox(MoveTemp(*FullPath));
 }
+
+bool ISandboxManager::CanLeaveSandbox() const
+{
+	if (!HasActiveSandbox())
+	{
+		return true;
+	}
+	
+	ISandboxLock* ActiveLock = GetActiveLock();
+	return !ActiveLock || ActiveLock->CanLeaveSandbox();
+}
 }
diff --git a/Engine/Plugins/Developer/Sandbox/FileSandbox/Source/FileSandboxCore/Private/Sandbox/SandboxManager.cpp b/Engine/Plugins/Developer/Sandbox/FileSandbox/Source/FileSandboxCore/Private/Sandbox/SandboxManager.cpp
--- a/Engine/Plugins/Developer/Sandbox/FileSandbox/Source/FileSandboxCore/Private/Sandbox/SandboxManager.cpp
+++ b/Engine/Plugins/Developer/Sandbox/FileSandbox/Source/FileSandboxCore/Private/Sandbox/SandboxManager.cpp
@@ -9,6 +9,7 @@
 #include "Misc/Paths.h"
 #include "Types/Manager/DeleteSandboxByDirectoryArgs.h"
 #include "Types/Manager/DeleteSandboxResult.h"
+#include "Types/Manager/LeaveSandboxResult.h"
 #include "Types/Manager/LoadSandboxByDirectoryArgs.h"
 #include "Types/Manager/LoadSandboxError.h"
 #include "Types/Manager/NewSandboxArgs.h"
@@ -38,18 +39,22 @@ FNewSandboxResult FSandboxManager::CreateNewSandbox(const FNewSandboxArgs& InArg
 		return { MakeError(FNewSandboxError{ ENewSandboxErrorReason::UnsuitablePath }) };
 	}
 
-	LeaveSandbox();
+	if (!LeaveSandbox())
+	{
+		return { MakeError(FNewSandboxError{ ENewSandboxErrorReason::CannotLeaveSandbox }) };
+	}
 	
 	const FString NewDirectory = FPaths::Combine(InArgs.SandboxBasePath, Name);
-	ActiveSandbox = FSandboxInstance::CreateNewSandbox(NewDirectory, InArgs.InitArgs, InArgs.MetaData);
-	ISandboxInstance* Instance = ActiveSandbox.Get();
-	if (!Instance)
+	TUniquePtr<FSandboxInstance> NewSandbox = FSandboxInstance::CreateNewSandbox(NewDirectory, InArgs.InitArgs, InArgs.MetaData);
+	ISandboxInstance* SandboxRaw = NewSandbox.Get();
+	if (!SandboxRaw)
 	{
 		return { MakeError(FNewSandboxError{ ENewSandboxErrorReason::IOError }) };
 	}
 	
-	OnPostSandboxStartupDelegate.Broadcast(*Instance);
-	return { MakeValue(Instance) };
+	ActiveSandboxData.Emplace(MoveTemp(NewSandbox), InArgs.InitArgs.Lock);
+	OnPostSandboxStartupDelegate.Broadcast(*SandboxRaw);
+	return { MakeValue(SandboxRaw) };
 }
 
 FLoadSandboxResult FSandboxManager::LoadSandbox(const FLoadSandboxByDirectoryArgs& InLoadArgs)
@@ -68,18 +73,22 @@ FLoadSandboxResult FSandboxManager::LoadSandbox(const FLoadSandboxByDirectoryArg
 	const FString LastDirectoryName = FPaths::GetPathLeaf(AbsoluteRoot);
 	UE_CLOG(LastDirectoryName.IsEmpty(), LogFileSandbox, Warning, TEXT("Sandbox placed %s will not have any name."), *AbsoluteRoot);
 	const TOptional<FString> SandboxName = LastDirectoryName.IsEmpty() ? LastDirectoryName : TOptional<FString>{};
-
-	LeaveSandbox();
 	
-	ActiveSandbox = FSandboxInstance::LoadSandbox(AbsoluteRoot, InLoadArgs.InitArgs);
-	ISandboxInstance* Instance = ActiveSandbox.Get();
-	if (!Instance)
+	if (!LeaveSandbox())
+	{
+		return { MakeError(FLoadSandboxError{ ELoadSandboxLoadErrorReason::CannotLeaveSandbox }) };
+	}
+	
+	TUniquePtr<FSandboxInstance> NewSandbox = FSandboxInstance::LoadSandbox(AbsoluteRoot, InLoadArgs.InitArgs);
+	ISandboxInstance* SandboxRaw = NewSandbox.Get();
+	if (!SandboxRaw)
 	{
 		return { MakeError(FLoadSandboxError{ ELoadSandboxLoadErrorReason::IOError }) };
 	}
 	
-	OnPostSandboxStartupDelegate.Broadcast(*Instance);
-	return { MakeValue(Instance) };
+	ActiveSandboxData.Emplace(MoveTemp(NewSandbox), InLoadArgs.InitArgs.Lock);
+	OnPostSandboxStartupDelegate.Broadcast(*SandboxRaw);
+	return { MakeValue(SandboxRaw) };
 }
 
 FDeleteSandboxResult FSandboxManager::DeleteSandbox(const FDeleteSandboxByDirectoryArgs& InArgs)
@@ -89,9 +98,11 @@ FDeleteSandboxResult FSandboxManager::DeleteSandbox(const FDeleteSandboxByDirect
 		return EDeleteSandboxErrorCode::InvalidDirectory;
 	}
 	
-	if (ActiveSandbox && FPaths::IsSamePath(ActiveSandbox->GetRootDirectory(), InArgs.Directory))
+	ISandboxInstance* ActiveSandbox = GetActiveSandboxInstance();
+	const bool bNeedsLeave = ActiveSandbox && FPaths::IsSamePath(ActiveSandbox->GetRootDirectory(), InArgs.Directory);
+	if (bNeedsLeave && !LeaveSandbox())
 	{
-		LeaveSandbox();
+		return EDeleteSandboxErrorCode::CannotLeaveSandbox;
 	}
 
 	IPlatformFile& PlatformFile = FPlatformFileManager::Get().GetPlatformFile();
@@ -100,14 +111,22 @@ FDeleteSandboxResult FSandboxManager::DeleteSandbox(const FDeleteSandboxByDirect
 		: EDeleteSandboxErrorCode::IOError;
 }
 
-void FSandboxManager::LeaveSandbox(const FLeaveSandboxArgs& InLeaveArgs)
+FLeaveSandboxResult FSandboxManager::LeaveSandbox(const FLeaveSandboxArgs& InLeaveArgs)
 {
-	if (ISandboxInstance* Instance = ActiveSandbox.Get())
+	if (!ActiveSandboxData)
+	{
+		return ELeaveSandboxErrorCode::Success;
+	}
+	
+	if (!CanLeaveSandbox())
 	{
-		OnPreSandboxShutdownDelegate.Broadcast(*Instance);
-		ActiveSandbox.Reset();
-		OnPostSandboxShutdownDelegate.Broadcast();
+		return ELeaveSandboxErrorCode::Locked;
 	}
+	
+	OnPreSandboxShutdownDelegate.Broadcast(*ActiveSandboxData->Sandbox.Get());
+	ActiveSandboxData.Reset();
+	OnPostSandboxShutdownDelegate.Broadcast();
+	return ELeaveSandboxErrorCode::Success;
 }
 
 void FSandboxManager::OnEngineExit()
diff --git a/Engine/Plugins/Developer/Sandbox/FileSandbox/Source/FileSandboxCore/Private/Sandbox/SandboxManager.h b/Engine/Plugins/Developer/Sandbox/FileSandbox/Source/FileSandboxCore/Private/Sandbox/SandboxManager.h
--- a/Engine/Plugins/Developer/Sandbox/FileSandbox/Source/FileSandboxCore/Private/Sandbox/SandboxManager.h
+++ b/Engine/Plugins/Developer/Sandbox/FileSandbox/Source/FileSandboxCore/Private/Sandbox/SandboxManager.h
@@ -20,17 +20,31 @@ public:
 	virtual FNewSandboxResult CreateNewSandbox(const FNewSandboxArgs& InArgs) override;
 	virtual FLoadSandboxResult LoadSandbox(const FLoadSandboxByDirectoryArgs& InLoadArgs) override;
 	virtual FDeleteSandboxResult DeleteSandbox(const FDeleteSandboxByDirectoryArgs& InArgs) override;
-	virtual void LeaveSandbox(const FLeaveSandboxArgs& InLeaveArgs = FLeaveSandboxArgs()) override;
-	virtual ISandboxInstance* GetActiveSandboxInstance() const override { return ActiveSandbox.Get(); }
+	virtual FLeaveSandboxResult LeaveSandbox(const FLeaveSandboxArgs& InLeaveArgs = FLeaveSandboxArgs()) override;
+	virtual ISandboxLock* GetActiveLock() const override { return ActiveSandboxData ? ActiveSandboxData->Lock.Get() : nullptr; }
+	virtual ISandboxInstance* GetActiveSandboxInstance() const override { return ActiveSandboxData ? ActiveSandboxData->Sandbox.Get() : nullptr; }
 	virtual FSandboxInstanceEvent& OnPostSandboxStartup() override { return OnPostSandboxStartupDelegate; }
 	virtual FSandboxInstanceEvent& OnPreSandboxShutdown() override { return OnPreSandboxShutdownDelegate; }
 	virtual FSandboxShutdownEvent& OnPostSandboxShutdown() override { return OnPostSandboxShutdownDelegate; }
 	//~ End ISandboxManager Interface
 
 private:
-
-	/** Current sandbox. Nullptr if not in any sandbox. */
-	TUniquePtr<FSandboxInstance> ActiveSandbox;
+	
+	struct FActiveSandboxData
+	{
+		/** Current sandbox. Nullptr if not in any sandbox. */
+		TUniquePtr<FSandboxInstance> Sandbox;
+		
+		/** Optional. If set, determines whether it is allowed to leave the sandbox. */
+		TSharedPtr<ISandboxLock> Lock;
+
+		explicit FActiveSandboxData(TUniquePtr<FSandboxInstance> Sandbox, TSharedPtr<ISandboxLock> InLock)
+			: Sandbox(MoveTemp(Sandbox))
+			, Lock(InLock)
+		{}
+	};
+	/** Set while in a sandbox */
+	TOptional<FActiveSandboxData> ActiveSandboxData;
 
 	/** Invoked after a sandbox has been successfully created. */
 	FSandboxInstanceEvent OnPostSandboxStartupDelegate;
diff --git a/Engine/Plugins/Developer/Sandbox/FileSandbox/Source/FileSandboxCore/Private/Types/Manager/LoadSandboxError.cpp b/Engine/Plugins/Developer/Sandbox/FileSandbox/Source/FileSandboxCore/Private/Types/Manager/LoadSandboxError.cpp
--- a/Engine/Plugins/Developer/Sandbox/FileSandbox/Source/FileSandboxCore/Private/Types/Manager/LoadSandboxError.cpp
+++ b/Engine/Plugins/Developer/Sandbox/FileSandbox/Source/FileSandboxCore/Private/Types/Manager/LoadSandboxError.cpp
@@ -8,12 +8,13 @@ namespace UE::FileSandboxCore
 {
 FString LexToString(ELoadSandboxLoadErrorReason InReason)
 {
-	static_assert(static_cast<int32>(ELoadSandboxLoadErrorReason::Count) == 4, "Update this switch");
+	static_assert(static_cast<int32>(ELoadSandboxLoadErrorReason::Count) == 5, "Update this switch");
 	switch (InReason)
 	{
 	case ELoadSandboxLoadErrorReason::InvalidDirectory: return TEXT("InvalidDirectory");
 	case ELoadSandboxLoadErrorReason::InvalidName: return TEXT("InvalidName");
 	case ELoadSandboxLoadErrorReason::IOError: return TEXT("IOError");
+	case ELoadSandboxLoadErrorReason::CannotLeaveSandbox: return TEXT("CannotLeaveSandbox");
 	case ELoadSandboxLoadErrorReason::Unspecified: return TEXT("Unspecified");
 	default: checkNoEntry(); return TEXT("Unknown");
 	}
diff --git a/Engine/Plugins/Developer/Sandbox/FileSandbox/Source/FileSandboxCore/Private/Types/Manager/NewSandboxError.cpp b/Engine/Plugins/Developer/Sandbox/FileSandbox/Source/FileSandboxCore/Private/Types/Manager/NewSandboxError.cpp
--- a/Engine/Plugins/Developer/Sandbox/FileSandbox/Source/FileSandboxCore/Private/Types/Manager/NewSandboxError.cpp
+++ b/Engine/Plugins/Developer/Sandbox/FileSandbox/Source/FileSandboxCore/Private/Types/Manager/NewSandboxError.cpp
@@ -9,11 +9,12 @@ namespace UE::FileSandboxCore
 {
 FString LexToString(ENewSandboxErrorReason InReason)
 {
-	static_assert(static_cast<int32>(ENewSandboxErrorReason::Count) == 3, "Update this switch");
+	static_assert(static_cast<int32>(ENewSandboxErrorReason::Count) == 4, "Update this switch");
 	switch (InReason)
 	{
 	case ENewSandboxErrorReason::UnsuitablePath: return TEXT("UnsuitablePath");
 	case ENewSandboxErrorReason::IOError: return TEXT("IOError");
+	case ENewSandboxErrorReason::CannotLeaveSandbox: return TEXT("CannotLeaveSandbox");
 	case ENewSandboxErrorReason::Unspecified: return TEXT("Unspecified");
 	default: checkNoEntry(); return TEXT("Unknown");
 	}
diff --git a/Engine/Plugins/Developer/Sandbox/FileSandbox/Source/FileSandboxCore/Public/ISandboxManager.h b/Engine/Plugins/Developer/Sandbox/FileSandbox/Source/FileSandboxCore/Public/ISandboxManager.h
--- a/Engine/Plugins/Developer/Sandbox/FileSandbox/Source/FileSandboxCore/Public/ISandboxManager.h
+++ b/Engine/Plugins/Developer/Sandbox/FileSandbox/Source/FileSandboxCore/Public/ISandboxManager.h
@@ -5,12 +5,14 @@
 #include "Delegates/Delegate.h"
 #include "Types/ArgsFwds.h"
 #include "Types/Manager/LeaveSandboxArgs.h"
+#include "Types/Manager/LeaveSandboxResult.h"
 #include "Types/Manager/SandboxCreationResult.h"
 
 #define UE_API FILESANDBOXCORE_API
 
 namespace UE::FileSandboxCore
 {
+class ISandboxLock;
 class ISandboxInstance;
 	
 DECLARE_MULTICAST_DELEGATE_OneParam(FSandboxInstanceEvent, ISandboxInstance&);
@@ -90,7 +92,18 @@ public:
 	
 	
 	/** Leaves the current sandbox. */
-	virtual void LeaveSandbox(const FLeaveSandboxArgs& InLeaveArgs = FLeaveSandboxArgs()) = 0;
+	virtual FLeaveSandboxResult LeaveSandbox(const FLeaveSandboxArgs& InLeaveArgs = FLeaveSandboxArgs()) = 0;
+	
+	/**
+	 * More advanced version of CanLeaveSandbox if you want to find out the reason why the sandbox cannot be left.
+	 * @return Lock object that determines whether the sandbox can be left.
+	 * @note Do not keep any raw reference to ISandboxLock as it turns stale when the sandbox is shutdown.
+	 */
+	virtual ISandboxLock* GetActiveLock() const = 0;
+	
+	/** @return Info about whether a LeaveSandbox call would be successful. */
+	UE_API bool CanLeaveSandbox() const;
+	
 
 	/**
 	 * Gets the currently active global sandbox instance.
diff --git a/Engine/Plugins/Developer/Sandbox/FileSandbox/Source/FileSandboxCore/Public/Interface/ISandboxLock.h b/Engine/Plugins/Developer/Sandbox/FileSandbox/Source/FileSandboxCore/Public/Interface/ISandboxLock.h
new file mode 100644
--- /dev/null
+++ b/Engine/Plugins/Developer/Sandbox/FileSandbox/Source/FileSandboxCore/Public/Interface/ISandboxLock.h
@@ -0,0 +1,32 @@
+﻿// Copyright Epic Games, Inc. All Rights Reserved.
+
+#pragma once
+
+class FText;
+
+namespace UE::FileSandboxCore
+{
+/** Interface to allow module to determine whether a sandbox can be left. */
+class ISandboxLock
+{
+public:
+	
+	/**
+	 * @param OutReason Optional. Reason set why the sandbox cannot be left.
+	 * @return Whether the sandbox can be left
+	 */
+	virtual bool CanLeaveSandbox(FText* OutReason = nullptr) const = 0;
+	
+	/** 
+	 * Asks the lock to be lifted. The request can be rejected.
+	 * @param OutReason Optional. Reason set when the lock could not be lifted.
+	 * @return Whether the lock was lifted
+	 */
+	virtual bool RequestLiftLock(FText* OutReason = nullptr) = 0;
+	
+	/** @return Whether it is supported to call RequestLiftLock. */
+	virtual bool SupportsLiftingLock() const = 0;
+	
+	virtual ~ISandboxLock() = default;
+};
+}
diff --git a/Engine/Plugins/Developer/Sandbox/FileSandbox/Source/FileSandboxCore/Public/Types/Manager/DeleteSandboxResult.h b/Engine/Plugins/Developer/Sandbox/FileSandbox/Source/FileSandboxCore/Public/Types/Manager/DeleteSandboxResult.h
--- a/Engine/Plugins/Developer/Sandbox/FileSandbox/Source/FileSandboxCore/Public/Types/Manager/DeleteSandboxResult.h
+++ b/Engine/Plugins/Developer/Sandbox/FileSandbox/Source/FileSandboxCore/Public/Types/Manager/DeleteSandboxResult.h
@@ -18,7 +18,13 @@ enum class EDeleteSandboxErrorCode : uint8
 	InvalidDirectory,
 
 	/** An I/O error occured, e.g. the directory could not be deleted. */
-	IOError
+	IOError,
+	
+	/**
+	 * Failed to leave the sandbox. Most likely because the sandbox is locked.
+	 * If you need more info consider calling LeaveSandbox explicitly before.
+	 */
+	CannotLeaveSandbox,
 };
 
 /** Describes the result of the delete operation, with a reason of why it failed. */
diff --git a/Engine/Plugins/Developer/Sandbox/FileSandbox/Source/FileSandboxCore/Public/Types/Manager/LeaveSandboxResult.h b/Engine/Plugins/Developer/Sandbox/FileSandbox/Source/FileSandboxCore/Public/Types/Manager/LeaveSandboxResult.h
new file mode 100644
--- /dev/null
+++ b/Engine/Plugins/Developer/Sandbox/FileSandbox/Source/FileSandboxCore/Public/Types/Manager/LeaveSandboxResult.h
@@ -0,0 +1,36 @@
+﻿// Copyright Epic Games, Inc. All Rights Reserved.
+
+#pragma once
+
+#include "HAL/Platform.h"
+#include "Misc/Optional.h"
+
+namespace UE::FileSandboxCore
+{
+enum class ELeaveSandboxErrorCode : uint8
+{
+	/** The sandbox was left successfully, or the engine was not in any sandbox. */
+	Success,
+	
+	/** 
+	 * The engine is locked to this sandbox. Inspect the ISandboxLock for the reason.
+	 * @see ISandboxManager::GetActiveLock
+	 */
+	Locked
+};
+
+/** Describes the result of the leave operation, with a reason of why it failed. */
+struct FLeaveSandboxResult
+{
+	/** Describes the result of the delete operation, with a reason of why it failed. */
+	ELeaveSandboxErrorCode ErrorCode = ELeaveSandboxErrorCode::Success;
+	
+	FLeaveSandboxResult() = default;
+	FLeaveSandboxResult(ELeaveSandboxErrorCode InErrorCode) : ErrorCode(InErrorCode) {}
+
+	/** @return Whether the sandbox was successfully left. */
+	bool IsSuccess() const { return ErrorCode == ELeaveSandboxErrorCode::Success; }
+	/** @return Whether the operation was successful. */
+	operator bool() const { return IsSuccess(); }
+};
+}
\ No newline at end of file
diff --git a/Engine/Plugins/Developer/Sandbox/FileSandbox/Source/FileSandboxCore/Public/Types/Manager/LoadSandboxError.h b/Engine/Plugins/Developer/Sandbox/FileSandbox/Source/FileSandboxCore/Public/Types/Manager/LoadSandboxError.h
--- a/Engine/Plugins/Developer/Sandbox/FileSandbox/Source/FileSandboxCore/Public/Types/Manager/LoadSandboxError.h
+++ b/Engine/Plugins/Developer/Sandbox/FileSandbox/Source/FileSandboxCore/Public/Types/Manager/LoadSandboxError.h
@@ -3,6 +3,8 @@
 #pragma once
 
 #include "HAL/Platform.h"
+#include "LeaveSandboxResult.h"
+#include "Misc/Optional.h"
 
 class FString;
 
@@ -17,6 +19,12 @@ enum class ELoadSandboxLoadErrorReason : uint8
 	
 	/** Some generic IO error occured (e.g. could not write manifest file, etc.) */
 	IOError,
+	
+	/**
+	 * Failed to leave the sandbox. Most likely because the sandbox is locked.
+	 * If you need more info consider calling LeaveSandbox explicitly before.
+	 */
+	CannotLeaveSandbox,
 
 	/** Dummy error code for initializing to a deterministic value. */
 	Unspecified,
diff --git a/Engine/Plugins/Developer/Sandbox/FileSandbox/Source/FileSandboxCore/Public/Types/Manager/NewSandboxError.h b/Engine/Plugins/Developer/Sandbox/FileSandbox/Source/FileSandboxCore/Public/Types/Manager/NewSandboxError.h
--- a/Engine/Plugins/Developer/Sandbox/FileSandbox/Source/FileSandboxCore/Public/Types/Manager/NewSandboxError.h
+++ b/Engine/Plugins/Developer/Sandbox/FileSandbox/Source/FileSandboxCore/Public/Types/Manager/NewSandboxError.h
@@ -3,6 +3,8 @@
 #pragma once
 
 #include "HAL/Platform.h"
+#include "LeaveSandboxResult.h"
+#include "Misc/Optional.h"
 
 class FString;
 
@@ -16,6 +18,12 @@ enum class ENewSandboxErrorReason : uint8
 	/** Some generic IO error occured (e.g. could not write manifest file, etc.) */
 	IOError,
 
+	/**
+	 * Failed to leave the sandbox. Most likely because the sandbox is locked.
+	 * If you need more info consider calling LeaveSandbox explicitly before.
+	 */
+	CannotLeaveSandbox,
+	
 	/** Dummy error code for initializing to a deterministic value. */
 	Unspecified,
 
diff --git a/Engine/Plugins/Developer/Sandbox/FileSandbox/Source/FileSandboxCore/Public/Types/Manager/SandboxInitArgs.h b/Engine/Plugins/Developer/Sandbox/FileSandbox/Source/FileSandboxCore/Public/Types/Manager/SandboxInitArgs.h
--- a/Engine/Plugins/Developer/Sandbox/FileSandbox/Source/FileSandboxCore/Public/Types/Manager/SandboxInitArgs.h
+++ b/Engine/Plugins/Developer/Sandbox/FileSandbox/Source/FileSandboxCore/Public/Types/Manager/SandboxInitArgs.h
@@ -3,10 +3,13 @@
 #pragma once
 
 #include "HAL/PlatformFileManager.h"
+#include "Templates/SharedPointer.h"
 #include "Types/SandboxInitFlags.h"
 
 namespace UE::FileSandboxCore
 {
+class ISandboxLock;
+
 /** Arguments for creating a sandbox. */
 struct FSandboxInitArgs
 {
@@ -18,6 +21,14 @@ struct FSandboxInitArgs
 	
 	/** Additional feature flags for the sandbox. */
 	ESandboxInitFlags Flags = ESandboxInitFlags::Default;
+	
+	/**
+	 * Optional. You can use this to force the engine to stay in the sandbox until you return true.
+	 * Whenever any API call attempts to make the engine leave the sandbox, this is used to check whether the sandbox is allowed to be left.
+	 * 
+	 * The sandbox system will only keep a reference to this lock for as long as the sandbox is active.
+	 */
+	TSharedPtr<ISandboxLock> Lock;
 
 	/** @return The platform file to use. */
 	IPlatformFile& GetPlatformFile() const { return OverridePlatformFile ? *OverridePlatformFile : FPlatformFileManager::Get().GetPlatformFile(); }
