File Sandbox: Create UI system for registering editor entry points for sandboxes.

This provides a generic widget overlay widget that editor systems (Multi User, Sandboxed Editing) can overlay into their own UI. For example, Multi-User can overlay this widget over its browser. The widget becomes visibile when a sandbox not owned by the editor system is active.

The widget shows a message to the user and optionally also adds a button to summon the external editor UI.

#Synced-CL 47637010

diff --git a/Engine/Plugins/Developer/Sandbox/FileSandbox/FileSandbox.uplugin b/Engine/Plugins/Developer/Sandbox/FileSandbox/FileSandbox.uplugin
--- a/Engine/Plugins/Developer/Sandbox/FileSandbox/FileSandbox.uplugin
+++ b/Engine/Plugins/Developer/Sandbox/FileSandbox/FileSandbox.uplugin
@@ -20,6 +20,11 @@
 			"Name": "FileSandboxCore",
 			"Type": "UncookedOnly",
 			"LoadingPhase": "Default"
+		},
+		{
+			"Name": "FileSandboxUI",
+			"Type": "Editor",
+			"LoadingPhase": "Default"
 		}
 	]
-}
+}
\ No newline at end of file
diff --git a/Engine/Plugins/Developer/Sandbox/FileSandbox/Source/FileSandboxUI/FileSandboxUI.Build.cs b/Engine/Plugins/Developer/Sandbox/FileSandbox/Source/FileSandboxUI/FileSandboxUI.Build.cs
new file mode 100644
--- /dev/null
+++ b/Engine/Plugins/Developer/Sandbox/FileSandbox/Source/FileSandboxUI/FileSandboxUI.Build.cs
@@ -0,0 +1,30 @@
+﻿// Copyright Epic Games, Inc. All Rights Reserved.
+
+using UnrealBuildTool;
+
+public class FileSandboxUI : ModuleRules
+{
+    public FileSandboxUI(ReadOnlyTargetRules Target) : base(Target)
+    {
+        PCHUsage = ModuleRules.PCHUsageMode.UseExplicitOrSharedPCHs;
+
+        PublicDependencyModuleNames.AddRange(
+            new string[]
+            {
+                "Core",
+                "Slate",
+            }
+        );
+
+        PrivateDependencyModuleNames.AddRange(
+            new string[]
+            {
+                "CoreUObject",
+                "Engine",
+                "FileSandboxCore",
+                "SlateCore",
+                "ToolWidgets"
+            }
+        );
+    }
+}
\ No newline at end of file
diff --git a/Engine/Plugins/Developer/Sandbox/FileSandbox/Source/FileSandboxUI/Private/FileSandboxStyle.cpp b/Engine/Plugins/Developer/Sandbox/FileSandbox/Source/FileSandboxUI/Private/FileSandboxStyle.cpp
new file mode 100644
--- /dev/null
+++ b/Engine/Plugins/Developer/Sandbox/FileSandbox/Source/FileSandboxUI/Private/FileSandboxStyle.cpp
@@ -0,0 +1,33 @@
+﻿// Copyright Epic Games, Inc. All Rights Reserved.
+
+#include "FileSandboxStyle.h"
+
+#include "Styling/SlateStyleMacros.h"
+#include "Styling/SlateStyleRegistry.h"
+
+namespace UE::FileSandboxUI
+{
+FFileSandboxStyle& FFileSandboxStyle::Get()
+{
+	static FFileSandboxStyle Style;
+	return Style;
+}
+
+FFileSandboxStyle::FFileSandboxStyle()
+	: FSlateStyleSet("FileSandboxStyle")
+{
+	const FString PluginContentDir = FPaths::EnginePluginsDir() / TEXT("Developer") / TEXT("Sandbox") / TEXT("FileSandboxCore") / TEXT("Content");
+	const FString EngineEditorSlateDir = FPaths::EngineContentDir() / TEXT("Editor") / TEXT("Slate");
+	SetContentRoot(PluginContentDir);
+	SetCoreContentRoot(EngineEditorSlateDir);
+	
+	Set("FileSandbox.TranslucentBackground", new FSlateColorBrush(FLinearColor(0.f, 0.f, 0.f, 0.6f)));
+	
+	FSlateStyleRegistry::RegisterSlateStyle(*this);
+}
+
+FFileSandboxStyle::~FFileSandboxStyle()
+{
+	FSlateStyleRegistry::UnRegisterSlateStyle(*this);
+}
+}
diff --git a/Engine/Plugins/Developer/Sandbox/FileSandbox/Source/FileSandboxUI/Private/FileSandboxStyle.h b/Engine/Plugins/Developer/Sandbox/FileSandbox/Source/FileSandboxUI/Private/FileSandboxStyle.h
new file mode 100644
--- /dev/null
+++ b/Engine/Plugins/Developer/Sandbox/FileSandbox/Source/FileSandboxUI/Private/FileSandboxStyle.h
@@ -0,0 +1,23 @@
+﻿// Copyright Epic Games, Inc. All Rights Reserved.
+
+#pragma once
+
+#include "Styling/ISlateStyle.h"
+#include "Styling/SlateStyle.h"
+
+class FSlateStyleSet;
+class ISlateStyle;
+
+namespace UE::FileSandboxUI
+{
+class FFileSandboxStyle final : public FSlateStyleSet
+{
+public:
+	
+	static FFileSandboxStyle& Get();
+
+	FFileSandboxStyle();
+	virtual ~FFileSandboxStyle() override;
+};
+}
+
diff --git a/Engine/Plugins/Developer/Sandbox/FileSandbox/Source/FileSandboxUI/Private/FileSandboxUIModule.cpp b/Engine/Plugins/Developer/Sandbox/FileSandbox/Source/FileSandboxUI/Private/FileSandboxUIModule.cpp
new file mode 100644
--- /dev/null
+++ b/Engine/Plugins/Developer/Sandbox/FileSandbox/Source/FileSandboxUI/Private/FileSandboxUIModule.cpp
@@ -0,0 +1,20 @@
+﻿// Copyright Epic Games, Inc. All Rights Reserved.
+
+#include "FileSandboxUIModule.h"
+
+#include "Mutex/SandboxEntryPointManager.h"
+
+namespace UE::FileSandboxUI
+{
+void FFileSandboxUIModule::StartupModule()
+{
+    EntryPoints = MakeUnique<FSandboxEntryPointManager>();
+}
+
+void FFileSandboxUIModule::ShutdownModule()
+{
+    EntryPoints.Reset();
+}
+}
+    
+IMPLEMENT_MODULE(UE::FileSandboxUI::FFileSandboxUIModule, FileSandboxUI)
\ No newline at end of file
diff --git a/Engine/Plugins/Developer/Sandbox/FileSandbox/Source/FileSandboxUI/Private/FileSandboxUIModule.h b/Engine/Plugins/Developer/Sandbox/FileSandbox/Source/FileSandboxUI/Private/FileSandboxUIModule.h
new file mode 100644
--- /dev/null
+++ b/Engine/Plugins/Developer/Sandbox/FileSandbox/Source/FileSandboxUI/Private/FileSandboxUIModule.h
@@ -0,0 +1,30 @@
+﻿// Copyright Epic Games, Inc. All Rights Reserved.
+
+#pragma once
+
+#include "IFileSandboxUIModule.h"
+#include "Mutex/SandboxEntryPointManager.h"
+
+namespace UE::FileSandboxUI
+{
+/** Knows about editor systems that start sandboxes. */
+class FFileSandboxUIModule : public IFileSandboxUIModule
+{
+public:
+	
+	//~ Begin IModuleInterface Interface
+	virtual void StartupModule() override;
+	virtual void ShutdownModule() override;
+	//~ End IModuleInterface Interface
+	
+	//~ Begin IFileSandboxUIModule Interface
+	virtual ISandboxEntryPointRegistry& GetEntryPointRegistry() override { return *EntryPoints.Get(); }
+	//~ End IFileSandboxUIModule Interface
+	
+private:
+	
+	/** Knows about editor systems that start sandboxes. */
+	TUniquePtr<FSandboxEntryPointManager> EntryPoints;
+};
+}
+
diff --git a/Engine/Plugins/Developer/Sandbox/FileSandbox/Source/FileSandboxUI/Private/Mutex/EntryPointWidgetFactory.cpp b/Engine/Plugins/Developer/Sandbox/FileSandbox/Source/FileSandboxUI/Private/Mutex/EntryPointWidgetFactory.cpp
new file mode 100644
--- /dev/null
+++ b/Engine/Plugins/Developer/Sandbox/FileSandbox/Source/FileSandboxUI/Private/Mutex/EntryPointWidgetFactory.cpp
@@ -0,0 +1,21 @@
+﻿// Copyright Epic Games, Inc. All Rights Reserved.
+
+#include "EntryPoint/EntryPointWidgetFactory.h"
+
+#include "ExternalSandboxActiveViewModel.h"
+#include "Widgets/SExternalSandboxActiveOverlay.h"
+
+namespace UE::FileSandboxUI
+{
+TSharedRef<SWidget> MakeExternalSandboxActiveOverlay(const TSharedRef<IExternalSandboxActiveViewModel>& InViewModel)
+{
+	return SNew(SExternalSandboxActiveOverlay, InViewModel);
+}
+
+TSharedRef<IExternalSandboxActiveViewModel> MakeExternalSandboxActiveViewModel(
+	const TSharedRef<ISandboxEntryPoint>& InOwningEntryPoint
+	)
+{
+	return MakeShared<FExternalSandboxActiveViewModel>(InOwningEntryPoint);
+}
+}
diff --git a/Engine/Plugins/Developer/Sandbox/FileSandbox/Source/FileSandboxUI/Private/Mutex/ExternalSandboxActiveViewModel.cpp b/Engine/Plugins/Developer/Sandbox/FileSandbox/Source/FileSandboxUI/Private/Mutex/ExternalSandboxActiveViewModel.cpp
new file mode 100644
--- /dev/null
+++ b/Engine/Plugins/Developer/Sandbox/FileSandbox/Source/FileSandboxUI/Private/Mutex/ExternalSandboxActiveViewModel.cpp
@@ -0,0 +1,93 @@
+﻿// Copyright Epic Games, Inc. All Rights Reserved.
+
+#include "ExternalSandboxActiveViewModel.h"
+
+#include "EntryPoint/ISandboxEntryPoint.h"
+#include "EntryPoint/ISandboxEntryPointRegistry.h"
+#include "IFileSandboxCoreModule.h"
+#include "IFileSandboxUIModule.h"
+#include "ISandboxManager.h"
+#include "Internationalization/Internationalization.h"
+#include "Internationalization/Text.h"
+
+#define LOCTEXT_NAMESPACE "FExternalSandboxActiveViewModel"
+
+namespace UE::FileSandboxUI
+{
+FExternalSandboxActiveViewModel::FExternalSandboxActiveViewModel(
+	const TSharedRef<ISandboxEntryPoint>& InOwningEntryPoint
+	)
+	: OwningEntryPoint(InOwningEntryPoint)
+{
+	using namespace UE::FileSandboxCore;
+	IFileSandboxCoreModule& ModuleCore = IFileSandboxCoreModule::Get();
+	ISandboxManager& SandboxManager = ModuleCore.GetSandboxManager();
+	SandboxManager.OnPostSandboxStartup().AddRaw(this, &FExternalSandboxActiveViewModel::OnSandboxStartup);
+	SandboxManager.OnPostSandboxShutdown().AddRaw(this, &FExternalSandboxActiveViewModel::OnSandboxShutdown);
+}
+
+FExternalSandboxActiveViewModel::~FExternalSandboxActiveViewModel()
+{
+	using namespace UE::FileSandboxCore;
+	if (IFileSandboxCoreModule::IsAvailable())
+	{
+		IFileSandboxCoreModule& ModuleCore = IFileSandboxCoreModule::Get();
+		ISandboxManager& SandboxManager = ModuleCore.GetSandboxManager();
+		SandboxManager.OnPostSandboxStartup().RemoveAll(this);
+		SandboxManager.OnPostSandboxShutdown().RemoveAll(this);
+	}
+}
+
+bool FExternalSandboxActiveViewModel::IsExternalSandboxActive() const
+{
+	using namespace UE::FileSandboxCore;
+	IFileSandboxCoreModule& ModuleCore = IFileSandboxCoreModule::Get();
+	IFileSandboxUIModule& ModuleUI = IFileSandboxUIModule::Get();
+	return ModuleCore.GetSandboxManager().HasActiveSandbox() 
+		&& ModuleUI.GetEntryPointRegistry().FindOwnerOfActiveSandbox() != OwningEntryPoint;
+}
+
+void FExternalSandboxActiveViewModel::SummonSandboxOwnerUI() const
+{
+	IFileSandboxUIModule& ModuleUI = IFileSandboxUIModule::Get();
+	if (const TSharedPtr<ISandboxEntryPoint> EntryPoint = ModuleUI.GetEntryPointRegistry().FindOwnerOfActiveSandbox())
+	{
+		EntryPoint->SummonProviderUI();
+	}
+}
+
+bool FExternalSandboxActiveViewModel::IsSummoningSupported() const
+{
+	IFileSandboxUIModule& ModuleUI = IFileSandboxUIModule::Get();
+	const TSharedPtr<ISandboxEntryPoint> EntryPoint = ModuleUI.GetEntryPointRegistry().FindOwnerOfActiveSandbox();
+	return EntryPoint && EntryPoint != OwningEntryPoint;
+}
+
+FText FExternalSandboxActiveViewModel::GetSummonActionLabel() const
+{
+	IFileSandboxUIModule& ModuleUI = IFileSandboxUIModule::Get();
+	const TSharedPtr<ISandboxEntryPoint> EntryPoint = ModuleUI.GetEntryPointRegistry().FindOwnerOfActiveSandbox();
+	return EntryPoint ? EntryPoint->GetEntryPointLabel() : FText::GetEmpty();
+}
+
+FText FExternalSandboxActiveViewModel::GetExternalSandboxActiveText() const
+{
+	IFileSandboxUIModule& ModuleUI = IFileSandboxUIModule::Get();
+	
+	if (const TSharedPtr<ISandboxEntryPoint> EntryPoint = ModuleUI.GetEntryPointRegistry().FindOwnerOfActiveSandbox())
+	{
+		return FText::Format(
+			LOCTEXT("Unavailable.KnownEntryPointFmt", "{0} is unavailable because the engine is sandboxed by {1}."),
+			OwningEntryPoint->GetEntryPointLabel(),
+			EntryPoint->GetEntryPointLabel()
+			);
+	}
+	
+	return FText::Format(
+		LOCTEXT("Unavailable.UnknownEntryPointFmt", "{0} is unavailable because the engine is sandboxed."),
+		OwningEntryPoint->GetEntryPointLabel()
+		);
+}
+}
+
+#undef LOCTEXT_NAMESPACE
\ No newline at end of file
diff --git a/Engine/Plugins/Developer/Sandbox/FileSandbox/Source/FileSandboxUI/Private/Mutex/ExternalSandboxActiveViewModel.h b/Engine/Plugins/Developer/Sandbox/FileSandbox/Source/FileSandboxUI/Private/Mutex/ExternalSandboxActiveViewModel.h
new file mode 100644
--- /dev/null
+++ b/Engine/Plugins/Developer/Sandbox/FileSandbox/Source/FileSandboxUI/Private/Mutex/ExternalSandboxActiveViewModel.h
@@ -0,0 +1,41 @@
+﻿// Copyright Epic Games, Inc. All Rights Reserved.
+
+#pragma once
+
+#include "ISandboxInstance.h"
+#include "EntryPoint/IExternalSandboxActiveViewModel.h"
+
+namespace UE::FileSandboxUI
+{
+class ISandboxEntryPoint;
+
+/** This view-model is associated with an ISandboxEntryPoint. The UI is shown when the sandbox is not owned by its entry point. */
+class FExternalSandboxActiveViewModel : public IExternalSandboxActiveViewModel
+{
+public:
+	
+	explicit FExternalSandboxActiveViewModel(const TSharedRef<ISandboxEntryPoint>& InOwningEntryPoint);
+	virtual ~FExternalSandboxActiveViewModel() override;
+	
+	//~ Begin IExternalSandboxActiveViewModel Interface
+	virtual bool IsExternalSandboxActive() const override;
+	virtual void SummonSandboxOwnerUI() const override;
+	virtual bool IsSummoningSupported() const override;
+	virtual FText GetSummonActionLabel() const override;
+	virtual FText GetExternalSandboxActiveText() const override;
+	virtual FSimpleMulticastDelegate& OnVisibilityChanged() override { return OnVisibilityChangedDelegate; }
+	//~ End IExternalSandboxActiveViewModel Interface
+	
+private:
+	
+	/** The entry point for which the UI is being displayed. Used to determine whether an external sandbox is active. */
+	const TSharedRef<ISandboxEntryPoint> OwningEntryPoint;
+	
+	/** Event invoked when the visibility of the widget should change. */
+	FSimpleMulticastDelegate OnVisibilityChangedDelegate;
+	
+	void OnSandboxStartup(FileSandboxCore::ISandboxInstance& SandboxInstance) { return OnVisibilityChangedDelegate.Broadcast(); }
+	void OnSandboxShutdown() { return OnVisibilityChangedDelegate.Broadcast(); }
+};
+}
+
diff --git a/Engine/Plugins/Developer/Sandbox/FileSandbox/Source/FileSandboxUI/Private/Mutex/SandboxEntryPointManager.cpp b/Engine/Plugins/Developer/Sandbox/FileSandbox/Source/FileSandboxUI/Private/Mutex/SandboxEntryPointManager.cpp
new file mode 100644
--- /dev/null
+++ b/Engine/Plugins/Developer/Sandbox/FileSandbox/Source/FileSandboxUI/Private/Mutex/SandboxEntryPointManager.cpp
@@ -0,0 +1,41 @@
+﻿// Copyright Epic Games, Inc. All Rights Reserved.
+
+#include "SandboxEntryPointManager.h"
+
+#include "EntryPoint/ISandboxEntryPoint.h"
+#include "IFileSandboxCoreModule.h"
+#include "ISandboxManager.h"
+
+namespace UE::FileSandboxUI
+{
+void FSandboxEntryPointManager::RegisterEntryPoint(const TSharedRef<ISandboxEntryPoint>& InEntryPoint)
+{
+	EntryPoints.AddUnique(InEntryPoint);
+	OnEntryPointsChangedDelegate.Broadcast();
+}
+
+void FSandboxEntryPointManager::UnregisterEntryPoint(const TSharedRef<ISandboxEntryPoint>& InEntryPoint)
+{
+	EntryPoints.RemoveSingle(InEntryPoint);
+	OnEntryPointsChangedDelegate.Broadcast();
+}
+
+TSharedPtr<ISandboxEntryPoint> FSandboxEntryPointManager::FindOwnerOfActiveSandbox() const
+{
+	using namespace UE::FileSandboxCore;
+	IFileSandboxCoreModule& Module = IFileSandboxCoreModule::Get();
+	const ISandboxInstance* Instance = Module.GetSandboxManager().GetActiveSandboxInstance();
+	if (!Instance)
+	{
+		return nullptr;
+	}
+	
+	// We don't handle multiple EntryPoints reporting ownership.
+	const int32 Index = EntryPoints.IndexOfByPredicate([Instance](const TWeakPtr<ISandboxEntryPoint>& WeakEntryPoint)
+	{
+		const TSharedPtr<ISandboxEntryPoint> EntryPointPin = WeakEntryPoint.Pin();
+		return EntryPointPin && EntryPointPin->OwnsSandbox(*Instance);
+	});
+	return EntryPoints.IsValidIndex(Index) ? EntryPoints[Index].Pin() : nullptr;
+}
+}
diff --git a/Engine/Plugins/Developer/Sandbox/FileSandbox/Source/FileSandboxUI/Private/Mutex/SandboxEntryPointManager.h b/Engine/Plugins/Developer/Sandbox/FileSandbox/Source/FileSandboxUI/Private/Mutex/SandboxEntryPointManager.h
new file mode 100644
--- /dev/null
+++ b/Engine/Plugins/Developer/Sandbox/FileSandbox/Source/FileSandboxUI/Private/Mutex/SandboxEntryPointManager.h
@@ -0,0 +1,30 @@
+﻿// Copyright Epic Games, Inc. All Rights Reserved.
+
+#pragma once
+
+#include "Containers/Array.h"
+#include "EntryPoint/ISandboxEntryPointRegistry.h"
+
+namespace UE::FileSandboxUI
+{
+class FSandboxEntryPointManager : public ISandboxEntryPointRegistry
+{
+public:
+	
+	//~ Begin ISandboxEntryPointRegistry Interface
+	virtual void RegisterEntryPoint(const TSharedRef<ISandboxEntryPoint>& InEntryPoint) override;
+	virtual void UnregisterEntryPoint(const TSharedRef<ISandboxEntryPoint>& InEntryPoint) override;
+	virtual TSharedPtr<ISandboxEntryPoint> FindOwnerOfActiveSandbox() const override;
+	virtual FSimpleMulticastDelegate& OnEntryPointsChanged() override { return OnEntryPointsChangedDelegate; }
+	//~ End ISandboxEntryPointRegistry Interface
+	
+private:
+	
+	/** The entry points this manager knows about. */
+	TArray<TWeakPtr<ISandboxEntryPoint>> EntryPoints;
+	
+	/** Event invoked when the registered entry points have changed. */
+	FSimpleMulticastDelegate OnEntryPointsChangedDelegate;
+};
+}
+
diff --git a/Engine/Plugins/Developer/Sandbox/FileSandbox/Source/FileSandboxUI/Private/Mutex/Widgets/SExternalSandboxActiveOverlay.cpp b/Engine/Plugins/Developer/Sandbox/FileSandbox/Source/FileSandboxUI/Private/Mutex/Widgets/SExternalSandboxActiveOverlay.cpp
new file mode 100644
--- /dev/null
+++ b/Engine/Plugins/Developer/Sandbox/FileSandbox/Source/FileSandboxUI/Private/Mutex/Widgets/SExternalSandboxActiveOverlay.cpp
@@ -0,0 +1,90 @@
+﻿// Copyright Epic Games, Inc. All Rights Reserved.
+
+#include "Mutex/Widgets/SExternalSandboxActiveOverlay.h"
+
+#include "Components/VerticalBox.h"
+#include "EntryPoint/IExternalSandboxActiveViewModel.h"
+#include "FileSandboxStyle.h"
+#include "SActionButton.h"
+#include "Widgets/Images/SImage.h"
+#include "Widgets/SOverlay.h"
+#include "Widgets/Text/STextBlock.h"
+
+#define LOCTEXT_NAMESPACE "SExternalSandboxActiveOverlay"
+
+namespace UE::FileSandboxUI
+{
+void SExternalSandboxActiveOverlay::Construct(const FArguments& InArgs, const TSharedRef<IExternalSandboxActiveViewModel>& InViewModel)
+{
+	ViewModel = InViewModel;
+	
+	ChildSlot
+	[
+		SNew(SOverlay)
+
+		+SOverlay::Slot()
+		.ZOrder(0)
+		[
+			SNew(SImage)
+			.Image(FFileSandboxStyle::Get().GetBrush("FileSandbox.TranslucentBackground"))
+		]
+
+		+SOverlay::Slot()
+		.ZOrder(1)
+		.HAlign(HAlign_Center)
+		.VAlign(VAlign_Center)
+		[
+			SNew(SVerticalBox)
+
+			+SVerticalBox::Slot()
+			.AutoHeight()
+			[
+				SNew(STextBlock)
+				.Text(this, &SExternalSandboxActiveOverlay::GetTitleText)
+				.Justification(ETextJustify::Center)
+			]
+
+			+SVerticalBox::Slot()
+			.AutoHeight()
+			.HAlign(HAlign_Center)
+			.Padding(0.f, 10.f)
+			[
+				SNew(SHorizontalBox)
+
+				+SHorizontalBox::Slot()
+				.AutoWidth()
+				[
+					SNew(SActionButton)
+					.Text(this, &SExternalSandboxActiveOverlay::GetSummonButtonLabel)
+					.Visibility_Lambda([this]{ return ViewModel->IsSummoningSupported() ? EVisibility::Visible : EVisibility::Collapsed; })
+					.OnClicked(this, &SExternalSandboxActiveOverlay::OnSummonClicked)
+				]
+			]
+		]
+	];
+	
+	// By setting visiblity here, we avoid duplicate visiblity logic for all API users.
+	SetVisibility(TAttribute<EVisibility>::CreateLambda([this]
+	{
+		return ViewModel->IsExternalSandboxActive() ? EVisibility::Visible : EVisibility::Collapsed;
+	}));
+}
+
+FText SExternalSandboxActiveOverlay::GetTitleText() const
+{
+	return ViewModel->GetExternalSandboxActiveText();
+}
+
+FText SExternalSandboxActiveOverlay::GetSummonButtonLabel() const
+{
+	return FText::Format(LOCTEXT("SummonUI.Label", "Go to {0}"), ViewModel->GetSummonActionLabel());
+}
+
+FReply SExternalSandboxActiveOverlay::OnSummonClicked()
+{
+	ViewModel->SummonSandboxOwnerUI();
+	return FReply::Handled();
+}
+}
+
+#undef LOCTEXT_NAMESPACE
\ No newline at end of file
diff --git a/Engine/Plugins/Developer/Sandbox/FileSandbox/Source/FileSandboxUI/Private/Mutex/Widgets/SExternalSandboxActiveOverlay.h b/Engine/Plugins/Developer/Sandbox/FileSandbox/Source/FileSandboxUI/Private/Mutex/Widgets/SExternalSandboxActiveOverlay.h
new file mode 100644
--- /dev/null
+++ b/Engine/Plugins/Developer/Sandbox/FileSandbox/Source/FileSandboxUI/Private/Mutex/Widgets/SExternalSandboxActiveOverlay.h
@@ -0,0 +1,35 @@
+﻿// Copyright Epic Games, Inc. All Rights Reserved.
+
+#pragma once
+
+#include "Widgets/DeclarativeSyntaxSupport.h"
+#include "Widgets/SCompoundWidget.h"
+
+namespace UE::FileSandboxUI
+{
+class IExternalSandboxActiveViewModel;
+
+/**
+ * Use this overlay in your editor system to tell the user that the engine is sandboxed by another editor system.
+ * This widget shows a button to summon the other system's UI.
+ */
+class SExternalSandboxActiveOverlay : public SCompoundWidget
+{
+public:
+	
+	SLATE_BEGIN_ARGS(SExternalSandboxActiveOverlay) {}
+	SLATE_END_ARGS()
+
+	void Construct(const FArguments& InArgs, const TSharedRef<IExternalSandboxActiveViewModel>& InViewModel);
+	
+private:
+	
+	/** View model that we're displaying. */
+	TSharedPtr<IExternalSandboxActiveViewModel> ViewModel;
+	
+	FText GetTitleText() const;
+	FText GetSummonButtonLabel() const;
+	FReply OnSummonClicked();
+};
+}
+
diff --git a/Engine/Plugins/Developer/Sandbox/FileSandbox/Source/FileSandboxUI/Public/EntryPoint/EntryPointWidgetFactory.h b/Engine/Plugins/Developer/Sandbox/FileSandbox/Source/FileSandboxUI/Public/EntryPoint/EntryPointWidgetFactory.h
new file mode 100644
--- /dev/null
+++ b/Engine/Plugins/Developer/Sandbox/FileSandbox/Source/FileSandboxUI/Public/EntryPoint/EntryPointWidgetFactory.h
@@ -0,0 +1,26 @@
+﻿// Copyright Epic Games, Inc. All Rights Reserved.
+
+#pragma once
+
+#include "Templates/SharedPointerFwd.h"
+
+#define UE_API FILESANDBOXUI_API
+
+class FText;
+class SWidget;
+
+namespace UE::FileSandboxUI
+{
+class ISandboxEntryPoint;
+class IExternalSandboxActiveViewModel;
+
+/** @return Widget to overlay your UI with when you want to tell the user that a sandbox owned by another editor system is blocking your feature. */
+UE_API TSharedRef<SWidget> MakeExternalSandboxActiveOverlay(const TSharedRef<IExternalSandboxActiveViewModel>& InViewModel);
+
+/** @return A view model for the entry point. */
+UE_API TSharedRef<IExternalSandboxActiveViewModel> MakeExternalSandboxActiveViewModel(
+	const TSharedRef<ISandboxEntryPoint>& InOwningEntryPoint
+	);
+}
+
+#undef UE_API
\ No newline at end of file
diff --git a/Engine/Plugins/Developer/Sandbox/FileSandbox/Source/FileSandboxUI/Public/EntryPoint/IExternalSandboxActiveViewModel.h b/Engine/Plugins/Developer/Sandbox/FileSandbox/Source/FileSandboxUI/Public/EntryPoint/IExternalSandboxActiveViewModel.h
new file mode 100644
--- /dev/null
+++ b/Engine/Plugins/Developer/Sandbox/FileSandbox/Source/FileSandboxUI/Public/EntryPoint/IExternalSandboxActiveViewModel.h
@@ -0,0 +1,36 @@
+﻿// Copyright Epic Games, Inc. All Rights Reserved.
+
+#pragma once
+
+#include "Delegates/Delegate.h"
+#include "Internationalization/Text.h"
+
+namespace UE::FileSandboxUI
+{
+/** 
+ * This view-model is used for UI telling the user that a sandbox cannot currently be started.
+ * Usually, the widget displays a message that another editor system has started a sandbox and provides a way to summon the UI of the other system. 
+ */
+class IExternalSandboxActiveViewModel
+{
+public:
+	
+	/** @return Whether the mutex widget should be displayed, i.e. a sandbox is active and not owned by the current ISandboxEntryPoint. */
+	virtual bool IsExternalSandboxActive() const = 0;
+	
+	/** Summons the UI of the system owning the active sandbox. */
+	virtual void SummonSandboxOwnerUI() const = 0;
+	/** @return Whether the UI (e.g. button) for summoning the sandbox owner should be shown. */
+	virtual bool IsSummoningSupported() const = 0;
+	/** @return The label to display for the summon action, e.g. could be displayed in button. */
+	virtual FText GetSummonActionLabel() const = 0;
+	
+	/** @return A text explaining the current system is not available while the external system is active. Empty of no external sandbox is active. */
+	virtual FText GetExternalSandboxActiveText() const = 0;
+	
+	/** Event invoked when the visibility of the widget should change. */
+	virtual FSimpleMulticastDelegate& OnVisibilityChanged() = 0;
+	
+	virtual ~IExternalSandboxActiveViewModel() = default;
+};
+}
diff --git a/Engine/Plugins/Developer/Sandbox/FileSandbox/Source/FileSandboxUI/Public/EntryPoint/ISandboxEntryPoint.h b/Engine/Plugins/Developer/Sandbox/FileSandbox/Source/FileSandboxUI/Public/EntryPoint/ISandboxEntryPoint.h
new file mode 100644
--- /dev/null
+++ b/Engine/Plugins/Developer/Sandbox/FileSandbox/Source/FileSandboxUI/Public/EntryPoint/ISandboxEntryPoint.h
@@ -0,0 +1,30 @@
+﻿// Copyright Epic Games, Inc. All Rights Reserved.
+
+#pragma once
+
+class FString;
+class FText;
+namespace UE::FileSandboxCore { class ISandboxInstance; }
+
+namespace UE::FileSandboxUI
+{
+/** Implemented by editor systems that makes use of the sandbox feature. */
+class ISandboxEntryPoint
+{
+public:
+	
+	/** Summons the UI in which the sandboxes  */
+	virtual void SummonProviderUI() = 0;
+	
+	/**
+	 * @return The name of the UI that will be summoned. This is displayed e.g. in a summon button. 
+	 * This should be an app name, like "Multi-User" or "Sandboxed Editing".
+	 */
+	virtual FText GetEntryPointLabel() const = 0;
+	
+	/** @return Whether this provider owns the active sandbox instance. */
+	virtual bool OwnsSandbox(const FileSandboxCore::ISandboxInstance& InSandbox) const = 0;
+	
+	virtual ~ISandboxEntryPoint() = default;
+};
+}
diff --git a/Engine/Plugins/Developer/Sandbox/FileSandbox/Source/FileSandboxUI/Public/EntryPoint/ISandboxEntryPointRegistry.h b/Engine/Plugins/Developer/Sandbox/FileSandbox/Source/FileSandboxUI/Public/EntryPoint/ISandboxEntryPointRegistry.h
new file mode 100644
--- /dev/null
+++ b/Engine/Plugins/Developer/Sandbox/FileSandbox/Source/FileSandboxUI/Public/EntryPoint/ISandboxEntryPointRegistry.h
@@ -0,0 +1,32 @@
+﻿// Copyright Epic Games, Inc. All Rights Reserved.
+
+#pragma once
+
+#include "Delegates/Delegate.h"
+#include "Templates/SharedPointerFwd.h"
+
+class FString;
+
+namespace UE::FileSandboxUI
+{
+class ISandboxEntryPoint;
+
+/** Knows about editor systems that start sandboxes. */
+class ISandboxEntryPointRegistry
+{
+public:
+	
+	/** Registers the entry point. */
+	virtual void RegisterEntryPoint(const TSharedRef<ISandboxEntryPoint>& InEntryPoint) = 0;
+	/** Unregisters the entry point. */
+	virtual void UnregisterEntryPoint(const TSharedRef<ISandboxEntryPoint>& InEntryPoint) = 0;
+	
+	/** Finds the entry point that owns the active sandbox. */
+	virtual TSharedPtr<ISandboxEntryPoint> FindOwnerOfActiveSandbox() const = 0;
+	
+	/** Event invoked when the registered entry points have changed. */
+	virtual FSimpleMulticastDelegate& OnEntryPointsChanged() = 0;
+	
+	virtual ~ISandboxEntryPointRegistry() = default;
+};
+}
diff --git a/Engine/Plugins/Developer/Sandbox/FileSandbox/Source/FileSandboxUI/Public/IFileSandboxUIModule.h b/Engine/Plugins/Developer/Sandbox/FileSandbox/Source/FileSandboxUI/Public/IFileSandboxUIModule.h
new file mode 100644
--- /dev/null
+++ b/Engine/Plugins/Developer/Sandbox/FileSandbox/Source/FileSandboxUI/Public/IFileSandboxUIModule.h
@@ -0,0 +1,34 @@
+// Copyright Epic Games, Inc. All Rights Reserved.
+
+#pragma once
+
+#include "Modules/ModuleInterface.h"
+#include "Modules/ModuleManager.h"
+
+namespace UE::FileSandboxUI
+{
+class ISandboxEntryPointRegistry;
+
+class IFileSandboxUIModule : public IModuleInterface
+{
+public:
+
+	/**
+	 * Singleton-like access to this module's interface.  This is just for convenience!
+	 * Beware of calling this during the shutdown phase, though.  Your module might have been unloaded already.
+	 *
+	 * @return Returns singleton instance, loading the module on demand if needed
+	 */
+	static IFileSandboxUIModule& Get() { return FModuleManager::LoadModuleChecked<IFileSandboxUIModule>("FileSandboxUI"); }
+
+	/**
+	 * Checks to see if this module is loaded and ready.  It is only valid to call Get() during shutdown if IsAvailable() returns true.
+	 *
+	 * @return True if the module is loaded and ready to use
+	 */
+	static bool IsAvailable() { return FModuleManager::Get().IsModuleLoaded("FileSandboxUI"); }
+	
+	/** @return Manager with which you can register editor entry points into the sandbox system. */
+	virtual ISandboxEntryPointRegistry& GetEntryPointRegistry() = 0;
+};
+}
