File Sandbox Core: Reorganize util files

#Synced-CL 47637010

diff --git a/Engine/Plugins/Developer/Sandbox/FileSandbox/Source/FileSandboxCore/Private/Repository/NaiveSandboxRepository.h b/Engine/Plugins/Developer/Sandbox/FileSandbox/Source/FileSandboxCore/Private/Repository/NaiveSandboxRepository.h
--- a/Engine/Plugins/Developer/Sandbox/FileSandbox/Source/FileSandboxCore/Private/Repository/NaiveSandboxRepository.h
+++ b/Engine/Plugins/Developer/Sandbox/FileSandbox/Source/FileSandboxCore/Private/Repository/NaiveSandboxRepository.h
@@ -2,11 +2,10 @@
 
 #pragma once
 
-#include "ISandboxInstance.h"
 #include "Data/SandboxMetaData.h"
+#include "ISandboxInstance.h"
 #include "ISandboxRepository.h"
 #include "Templates/UnrealTemplate.h"
-#include "Utils/ScopedWatchedDirectory.h"
 
 struct FFileChangeData;
 
diff --git a/Engine/Plugins/Developer/Sandbox/FileSandbox/Source/FileSandboxCore/Private/Repository/WatchedSandboxRepository.cpp b/Engine/Plugins/Developer/Sandbox/FileSandbox/Source/FileSandboxCore/Private/Repository/WatchedSandboxRepository.cpp
--- a/Engine/Plugins/Developer/Sandbox/FileSandbox/Source/FileSandboxCore/Private/Repository/WatchedSandboxRepository.cpp
+++ b/Engine/Plugins/Developer/Sandbox/FileSandbox/Source/FileSandboxCore/Private/Repository/WatchedSandboxRepository.cpp
@@ -3,18 +3,16 @@
 #if UE_SANDBOX_WITH_DIRECTORY_WATCHER
 #include "WatchedSandboxRepository.h"
 
-#include "IDirectoryWatcher.h"
-#include "IDirectoryWatcher.h"
-#include "LogFileSandbox.h"
 #include "Algo/AnyOf.h"
 #include "GenericPlatform/GenericPlatformFile.h"
 #include "HAL/PlatformFileManager.h"
+#include "IDirectoryWatcher.h"
 #include "Misc/Paths.h"
 #include "Types/EBreakBehavior.h"
 #include "Types/RepositoryChangedEvent.h"
-#include "Utils/FileSandboxUtils.h"
 #include "Utils/SandboxDirectoryUtils.h"
-#include "Utils/SandboxUtils.h"
+#include "Utils/SandboxFileUtils.h"
+#include "Utils/Watcher/DirectoryWatcherUtils.h"
 
 namespace UE::FileSandboxCore
 {
diff --git a/Engine/Plugins/Developer/Sandbox/FileSandbox/Source/FileSandboxCore/Private/Repository/WatchedSandboxRepository.h b/Engine/Plugins/Developer/Sandbox/FileSandbox/Source/FileSandboxCore/Private/Repository/WatchedSandboxRepository.h
--- a/Engine/Plugins/Developer/Sandbox/FileSandbox/Source/FileSandboxCore/Private/Repository/WatchedSandboxRepository.h
+++ b/Engine/Plugins/Developer/Sandbox/FileSandbox/Source/FileSandboxCore/Private/Repository/WatchedSandboxRepository.h
@@ -7,7 +7,7 @@
 #include "ISandboxRepository.h"
 #include "Templates/UnrealTemplate.h"
 #include "Types/SandboxMetaInfo.h"
-#include "Utils/ScopedWatchedDirectory.h"
+#include "Utils/Watcher/ScopedWatchedDirectory.h"
 
 struct FFileChangeData;
 
diff --git a/Engine/Plugins/Developer/Sandbox/FileSandbox/Source/FileSandboxCore/Private/Sandbox/Persist/DiscardUtils.cpp b/Engine/Plugins/Developer/Sandbox/FileSandbox/Source/FileSandboxCore/Private/Sandbox/Persist/DiscardUtils.cpp
new file mode 100644
--- /dev/null
+++ b/Engine/Plugins/Developer/Sandbox/FileSandbox/Source/FileSandboxCore/Private/Sandbox/Persist/DiscardUtils.cpp
@@ -0,0 +1,106 @@
+﻿// Copyright Epic Games, Inc. All Rights Reserved.
+
+#include "DiscardUtils.h"
+
+#include "Containers/Array.h"
+#include "Containers/ArrayView.h"
+#include "Containers/Set.h"
+#include "Containers/UnrealString.h"
+#include "GenericPlatform/GenericPlatformFile.h"
+#include "Sandbox/Platform/SandboxMountPoint.h"
+#include "UObject/NameTypes.h"
+#include "Utils/FileChange/FileChange.h"
+#include "Utils/AssetUtils.h"
+
+namespace UE::FileSandboxCore
+{
+void DiscardRemovedPaths(
+	const TSet<FSandboxPath>& InDeletedSandboxPaths, IPlatformFile& InLowerLevel, 
+	TArray<FName>& OutPackagesPendingPurge,
+	TFunctionRef<void(const FFileChange&)> InProcessFileChangeCallback
+	)
+{
+	for (const FSandboxPath& NonSandboxPath : InDeletedSandboxPaths)
+	{
+		// If this file maps to a package then we need to flush its linker so that we can remove the file from the sandbox
+		FName PackageName;
+		FlushPackageFile(NonSandboxPath, &PackageName);
+
+		if (InLowerLevel.FileExists(*NonSandboxPath))
+		{
+			if (!PackageName.IsNone())
+			{
+				OutPackagesPendingPurge.Add(PackageName);
+			}
+			
+			InProcessFileChangeCallback(FFileChange{ NonSandboxPath, EFileChangeAction::Added });
+		}
+	}
+}
+
+void DiscardAddedAndModifiedPaths(
+	TConstArrayView<FSandboxMountPoint> InMountPoints, IPlatformFile& InLowerLevel, 
+	TArray<FName>& OutPackagesPendingHotReload,
+	TArray<FName>& OutPackagesPendingPurge,
+	TFunctionRef<void(const FFileChange&)> InProcessFileChangeCallback
+	)
+{
+	for (const FSandboxMountPoint& SandboxMountPoint : InMountPoints)
+	{
+		// Modify any files that exist in both the non-sandbox and sandbox directories
+		// Delete any files that exist in the sandbox but don't exist in the non-sandbox directory
+		InLowerLevel.IterateDirectoryRecursively(*SandboxMountPoint.Path.GetSandboxPath(), 
+			[&InLowerLevel, &SandboxMountPoint, &OutPackagesPendingHotReload, &OutPackagesPendingPurge, &InProcessFileChangeCallback]
+			(const TCHAR* InFilenameOrDirectory, const bool InIsDirectory) -> bool
+		{
+			if (!InIsDirectory)
+			{
+				const FSandboxedPlatformFilePath RemappedFilePath = FSandboxedPlatformFilePath::CreateNonSandboxPath(FPaths::ConvertRelativePathToFull(InFilenameOrDirectory), SandboxMountPoint.Path);
+
+				// If this file maps to a package then we need to flush its linker so that we can remove the file from the sandbox
+				FName PackageName;
+				FlushPackageFile(RemappedFilePath.GetNonSandboxPath(), &PackageName, false);
+
+				if (InLowerLevel.FileExists(*RemappedFilePath.GetNonSandboxPath()))
+				{
+					if (!PackageName.IsNone())
+					{
+						OutPackagesPendingHotReload.Add(PackageName);
+						OutPackagesPendingPurge.Remove(PackageName);
+					}
+					
+					InProcessFileChangeCallback(FFileChange{ RemappedFilePath.GetNonSandboxPath(), EFileChangeAction::Modified });
+				}
+				else
+				{
+					if (!PackageName.IsNone())
+					{
+						OutPackagesPendingPurge.Add(PackageName);
+						OutPackagesPendingHotReload.Remove(PackageName);
+					}
+					
+					InProcessFileChangeCallback(FFileChange{ RemappedFilePath.GetNonSandboxPath(), EFileChangeAction::Removed });
+				}
+			}
+			return true; // Continue iteration
+		});
+
+		//FString SandboxRootPath = GetSandboxRootPath();
+		//FString Stashed = SandboxRootPath.Replace(TEXT("Sandbox"), TEXT("Stashed"));
+		//LowerLevel->MoveFile(SandboxRootPath, Stashed);
+		// Delete everything under the mount point
+		//LowerLevel->IterateDirectory(*SandboxMountPoint.Path.GetSandboxPath(), [this](const TCHAR* InFilenameOrDirectory, const bool InIsDirectory) -> bool
+		//{
+		//	if (InIsDirectory)
+		//	{
+		//		LowerLevel->DeleteDirectoryRecursively(InFilenameOrDirectory);
+		//	}
+		//	else
+		//	{
+		//		LowerLevel->DeleteFile(InFilenameOrDirectory);
+		//	}
+		//	return true; // Continue iteration
+		//});
+	}
+}
+}
diff --git a/Engine/Plugins/Developer/Sandbox/FileSandbox/Source/FileSandboxCore/Private/Sandbox/Persist/DiscardUtils.h b/Engine/Plugins/Developer/Sandbox/FileSandbox/Source/FileSandboxCore/Private/Sandbox/Persist/DiscardUtils.h
new file mode 100644
--- /dev/null
+++ b/Engine/Plugins/Developer/Sandbox/FileSandbox/Source/FileSandboxCore/Private/Sandbox/Persist/DiscardUtils.h
@@ -0,0 +1,28 @@
+﻿// Copyright Epic Games, Inc. All Rights Reserved.
+
+#pragma once
+
+#include "Containers/ContainersFwd.h"
+#include "Sandbox/Platform/SandboxedPlatformFilePath.h"
+#include "Templates/Function.h"
+
+class IPlatformFile;
+
+namespace UE::FileSandboxCore
+{
+struct FFileChange;
+struct FSandboxMountPoint;
+	
+void DiscardRemovedPaths(
+	const TSet<FSandboxPath>& InDeletedSandboxPaths, IPlatformFile& InLowerLevel, 
+	TArray<FName>& OutPackagesPendingPurge,
+	TFunctionRef<void(const FFileChange&)> InProcessFileChangeCallback = [](const FFileChange&){}
+	);
+	
+void DiscardAddedAndModifiedPaths(
+	TConstArrayView<FSandboxMountPoint> InMountPoints, IPlatformFile& InLowerLevel, 
+	TArray<FName>& OutPackagesPendingHotReload,
+	TArray<FName>& OutPackagesPendingPurge,
+	TFunctionRef<void(const FFileChange&)> InProcessFileChangeCallback = [](const FFileChange&){}
+	);
+}
diff --git a/Engine/Plugins/Developer/Sandbox/FileSandbox/Source/FileSandboxCore/Private/Sandbox/Persist/PersistUtils.cpp b/Engine/Plugins/Developer/Sandbox/FileSandbox/Source/FileSandboxCore/Private/Sandbox/Persist/PersistUtils.cpp
--- a/Engine/Plugins/Developer/Sandbox/FileSandbox/Source/FileSandboxCore/Private/Sandbox/Persist/PersistUtils.cpp
+++ b/Engine/Plugins/Developer/Sandbox/FileSandbox/Source/FileSandboxCore/Private/Sandbox/Persist/PersistUtils.cpp
@@ -8,7 +8,7 @@
 #include "SourceControlOperations.h"
 #include "GenericPlatform/GenericPlatformFile.h"
 #include "Misc/Paths.h"
-#include "Utils/SandboxUtils.h"
+#include "Utils/AssetUtils.h"
 
 namespace UE::FileSandboxCore
 {
diff --git a/Engine/Plugins/Developer/Sandbox/FileSandbox/Source/FileSandboxCore/Private/Sandbox/Platform/SandboxPlatformFile.cpp b/Engine/Plugins/Developer/Sandbox/FileSandbox/Source/FileSandboxCore/Private/Sandbox/Platform/SandboxPlatformFile.cpp
--- a/Engine/Plugins/Developer/Sandbox/FileSandbox/Source/FileSandboxCore/Private/Sandbox/Platform/SandboxPlatformFile.cpp
+++ b/Engine/Plugins/Developer/Sandbox/FileSandbox/Source/FileSandboxCore/Private/Sandbox/Platform/SandboxPlatformFile.cpp
@@ -8,7 +8,6 @@
 #include "GenericPlatform/GenericPlatformFile.h"
 #include "HAL/PlatformFileManager.h"
 #include "ISourceControlProvider.h"
-#include "ISourceControlState.h"
 #include "JsonObjectConverter.h"
 #include "LogFileSandbox.h"
 #include "Logging/LogVerbosity.h"
@@ -18,24 +17,27 @@
 #include "Misc/ScopeLock.h"
 #include "Modules/ModuleManager.h"
 #include "RenderingThread.h"
-#include "SandboxMountPoint.h"
-#include "SourceControlOperations.h"
+#include "Sandbox/Persist/DiscardUtils.h"
 #include "Sandbox/Persist/PersistUtils.h"
-#include "Types/SandboxFileChange.h"
+#include "SandboxMountPoint.h"
 #include "Types/Sandbox/PersistArgs.h"
 #include "Types/Sandbox/PersistResult.h"
+#include "Types/SandboxFileChange.h"
 #include "UObject/LinkerLoad.h"
 #include "UObject/Package.h"
+#include "Utils/AssetUtils.h"
 #include "Utils/FileChange/FileChange.h"
-#include "Utils/Persist/DiscardUtils.h"
-#include "Utils/SandboxUtils.h"
+
+#if UE_SANDBOX_WITH_DIRECTORY_WATCHER
+#include "Utils/Watcher/DirectoryWatcherUtils.h"
+#include "IDirectoryWatcher.h"
+#endif
 
 #if WITH_EDITOR
 #include "DirectoryWatcherModule.h"
 #include "Editor.h"
 #include "Editor/UnrealEdEngine.h"
 #include "FileHelpers.h"
-#include "IDirectoryWatcher.h"
 #include "Selection.h"
 #endif
 
diff --git a/Engine/Plugins/Developer/Sandbox/FileSandbox/Source/FileSandboxCore/Private/Sandbox/SandboxInstance.cpp b/Engine/Plugins/Developer/Sandbox/FileSandbox/Source/FileSandboxCore/Private/Sandbox/SandboxInstance.cpp
--- a/Engine/Plugins/Developer/Sandbox/FileSandbox/Source/FileSandboxCore/Private/Sandbox/SandboxInstance.cpp
+++ b/Engine/Plugins/Developer/Sandbox/FileSandbox/Source/FileSandboxCore/Private/Sandbox/SandboxInstance.cpp
@@ -7,15 +7,14 @@
 #include "Persist/AggregatePersistErrorHandler.h"
 #include "Persist/AppendTextErrorHandler.h"
 #include "Persist/LogPersistErrorHandler.h"
-#include "Types/SandboxFileChange.h"
-#include "Types/Sandbox/PersistArgs.h"
 #include "Types/Manager/SandboxInitArgs.h"
 #include "Types/Sandbox/DiscardResult.h"
+#include "Types/SandboxFileChange.h"
+#include "Utils/AssetUtils.h"
+#include "Utils/FileChange/FileChange.h"
 #include "Utils/PackageSandboxUtils.h"
-#include "Utils/FileSandboxUtils.h"
 #include "Utils/SandboxDirectoryUtils.h"
-#include "Utils/SandboxUtils.h"
-#include "Utils/FileChange/FileChange.h"
+#include "Utils/SandboxFileUtils.h"
 
 namespace UE::FileSandboxCore
 {
diff --git a/Engine/Plugins/Developer/Sandbox/FileSandbox/Source/FileSandboxCore/Private/Sandbox/SandboxManager.cpp b/Engine/Plugins/Developer/Sandbox/FileSandbox/Source/FileSandboxCore/Private/Sandbox/SandboxManager.cpp
--- a/Engine/Plugins/Developer/Sandbox/FileSandbox/Source/FileSandboxCore/Private/Sandbox/SandboxManager.cpp
+++ b/Engine/Plugins/Developer/Sandbox/FileSandbox/Source/FileSandboxCore/Private/Sandbox/SandboxManager.cpp
@@ -15,7 +15,7 @@
 #include "Types/Manager/NewSandboxArgs.h"
 #include "Types/Manager/NewSandboxError.h"
 #include "Types/Manager/SandboxCreationResult.h"
-#include "Utils/FileSandboxUtils.h"
+#include "Utils/SandboxFileUtils.h"
 
 namespace UE::FileSandboxCore
 {
diff --git a/Engine/Plugins/Developer/Sandbox/FileSandbox/Source/FileSandboxCore/Private/Utils/AssetUtils.cpp b/Engine/Plugins/Developer/Sandbox/FileSandbox/Source/FileSandboxCore/Private/Utils/AssetUtils.cpp
new file mode 100644
--- /dev/null
+++ b/Engine/Plugins/Developer/Sandbox/FileSandbox/Source/FileSandboxCore/Private/Utils/AssetUtils.cpp
@@ -0,0 +1,111 @@
+﻿// Copyright Epic Games, Inc. All Rights Reserved.
+
+#include "AssetUtils.h"
+
+#include "HAL/FileManager.h"
+#include "Modules/ModuleManager.h"
+#include "UObject/Linker.h"
+#include "UObject/LinkerLoad.h"
+#include "UObject/Package.h"
+#include "UObject/UObjectGlobals.h"
+#include "Utils/PackageSandboxUtils.h"
+
+#if UE_SANDBOX_WITH_DIRECTORY_WATCHER
+#include "DirectoryWatcherModule.h"
+#include "IDirectoryWatcher.h"
+#include "Watcher/DirectoryWatcherUtils.h"
+#endif
+
+namespace UE::FileSandboxCore
+{
+void ReplayChanges(TConstArrayView<FString> DeletedFiles, TConstArrayView<FString> ModifiedFiles, TConstArrayView<FString> AddedFiles)
+{
+	// Causes any open editors to refresh
+	TArray<FName> PackagesPendingHotReload;
+	for (const FString& Modified : ModifiedFiles)
+	{
+		FString PackageName;
+		if (FPackageName::TryConvertFilenameToLongPackageName(Modified, PackageName))
+		{
+			PackagesPendingHotReload.Add(FName(*PackageName));
+		}
+	}
+	HotReloadPackages(PackagesPendingHotReload);
+
+	// This closes any open editors and removes the files from the content browser
+	TArray<FName> PackagesPendingPurge;
+	IFileManager& FileManager = IFileManager::Get();
+	for (const FString& Deleted : DeletedFiles)
+	{
+		FString PackageName;
+		if (FPackageName::TryConvertFilenameToLongPackageName(Deleted, PackageName))
+		{
+			PackagesPendingPurge.Add(FName(*PackageName));
+			FileManager.Delete(*Deleted);
+		}
+	}
+	PurgePackages(PackagesPendingPurge);
+
+	// This is needed to make added files show up in the content browser
+#if UE_SANDBOX_WITH_DIRECTORY_WATCHER
+	TArray<FFileChangeData> FileChanges;
+	Algo::Transform(AddedFiles, FileChanges, [](const FString& File){ return FFileChangeData(File, FFileChangeData::FCA_Added); });
+	if (FDirectoryWatcherModule* DirectoryWatcherModule = GetDirectoryWatcherModuleIfLoaded())
+	{
+		DirectoryWatcherModule->RegisterExternalChanges(FileChanges);
+	}
+#endif
+
+	SynchronizeAssetRegistry();
+}
+	
+void SynchronizeAssetRegistry()
+{
+#if UE_SANDBOX_WITH_DIRECTORY_WATCHER
+	IDirectoryWatcher* DirectoryWatcher = GetDirectoryWatcherIfLoaded();
+	if (!DirectoryWatcher)
+	{
+		return;
+	}
+
+	DirectoryWatcher->Tick(0.0f);
+#endif // WITH_EDITOR
+}
+
+void FlushPackageLoading(const FString& InPackageName, bool bForceBulkDataLoad)
+{
+	UPackage* ExistingPackage = FindPackage(nullptr, *InPackageName);
+	if (ExistingPackage)
+	{
+		if (!ExistingPackage->IsFullyLoaded())
+		{
+			FlushAsyncLoading();
+			ExistingPackage->FullyLoad();
+		}
+
+		if (bForceBulkDataLoad)
+		{
+			ResetLoaders(ExistingPackage);
+		}
+		else if (ExistingPackage->GetLinker())
+		{
+			ExistingPackage->GetLinker()->Detach();
+		}
+	}
+}
+
+bool FlushPackageFile(const FString& InFilename, FName* OutPackageName, bool bForceLoad)
+{
+	FString PackageName;
+	if (FPackageName::TryConvertFilenameToLongPackageName(InFilename, PackageName))
+	{
+		FlushPackageLoading(PackageName, bForceLoad);
+		if (OutPackageName)
+		{
+			*OutPackageName = *PackageName;
+		}
+		return true;
+	}
+	return false;
+}
+}
diff --git a/Engine/Plugins/Developer/Sandbox/FileSandbox/Source/FileSandboxCore/Private/Utils/AssetUtils.h b/Engine/Plugins/Developer/Sandbox/FileSandbox/Source/FileSandboxCore/Private/Utils/AssetUtils.h
new file mode 100644
--- /dev/null
+++ b/Engine/Plugins/Developer/Sandbox/FileSandbox/Source/FileSandboxCore/Private/Utils/AssetUtils.h
@@ -0,0 +1,24 @@
+﻿// Copyright Epic Games, Inc. All Rights Reserved.
+
+#pragma once
+
+#include "Containers/ContainersFwd.h"
+
+class FDirectoryWatcherModule;
+class FName;
+class FString;
+class IDirectoryWatcher;
+
+namespace UE::FileSandboxCore
+{
+/** Utility for applying changes from manifest file to the current IPlatformFile. */
+void ReplayChanges(TConstArrayView<FString> DeletedFiles, TConstArrayView<FString> ModifiedFiles, TConstArrayView<FString> AddedFiles);
+	
+/** Refreshes the assets registry after making underlying file changes. */
+void SynchronizeAssetRegistry();
+
+/** Finishes loading the package, if currently async loading, and detaches the linker. */
+void FlushPackageLoading(const FString& InPackageName, bool bForceBulkDataLoad = true);
+/** Alternate variant that flushes package loading by filename instead. */
+bool FlushPackageFile(const FString& InFilename, FName* OutPackageName = nullptr, bool bForceLoad = true);
+}
diff --git a/Engine/Plugins/Developer/Sandbox/FileSandbox/Source/FileSandboxCore/Private/Utils/FileSandboxUtils.cpp b/Engine/Plugins/Developer/Sandbox/FileSandbox/Source/FileSandboxCore/Private/Utils/FileSandboxUtils.cpp
deleted file mode 100644
diff --git a/Engine/Plugins/Developer/Sandbox/FileSandbox/Source/FileSandboxCore/Private/Utils/FileSandboxUtils.h b/Engine/Plugins/Developer/Sandbox/FileSandbox/Source/FileSandboxCore/Private/Utils/FileSandboxUtils.h
deleted file mode 100644
diff --git a/Engine/Plugins/Developer/Sandbox/FileSandbox/Source/FileSandboxCore/Private/Utils/Persist/DiscardUtils.cpp b/Engine/Plugins/Developer/Sandbox/FileSandbox/Source/FileSandboxCore/Private/Utils/Persist/DiscardUtils.cpp
deleted file mode 100644
diff --git a/Engine/Plugins/Developer/Sandbox/FileSandbox/Source/FileSandboxCore/Private/Utils/Persist/DiscardUtils.h b/Engine/Plugins/Developer/Sandbox/FileSandbox/Source/FileSandboxCore/Private/Utils/Persist/DiscardUtils.h
deleted file mode 100644
diff --git a/Engine/Plugins/Developer/Sandbox/FileSandbox/Source/FileSandboxCore/Private/Utils/SandboxDirectoryUtils.cpp b/Engine/Plugins/Developer/Sandbox/FileSandbox/Source/FileSandboxCore/Private/Utils/SandboxDirectoryUtils.cpp
--- a/Engine/Plugins/Developer/Sandbox/FileSandbox/Source/FileSandboxCore/Private/Utils/SandboxDirectoryUtils.cpp
+++ b/Engine/Plugins/Developer/Sandbox/FileSandbox/Source/FileSandboxCore/Private/Utils/SandboxDirectoryUtils.cpp
@@ -2,15 +2,15 @@
 
 #include "Utils/SandboxDirectoryUtils.h"
 
-#include "FileSandboxUtils.h"
-#include "JsonObjectConverter.h"
-#include "LogFileSandbox.h"
 #include "Data/ManifestData.h"
 #include "Data/SandboxMetaData.h"
 #include "Data/VersionInfo.h"
 #include "HAL/PlatformFileManager.h"
+#include "JsonObjectConverter.h"
+#include "LogFileSandbox.h"
 #include "Misc/FileHelper.h"
 #include "Misc/Paths.h"
+#include "SandboxFileUtils.h"
 #include "Templates/Function.h"
 #include "Types/EBreakBehavior.h"
 
diff --git a/Engine/Plugins/Developer/Sandbox/FileSandbox/Source/FileSandboxCore/Private/Utils/SandboxFileUtils.cpp b/Engine/Plugins/Developer/Sandbox/FileSandbox/Source/FileSandboxCore/Private/Utils/SandboxFileUtils.cpp
new file mode 100644
--- /dev/null
+++ b/Engine/Plugins/Developer/Sandbox/FileSandbox/Source/FileSandboxCore/Private/Utils/SandboxFileUtils.cpp
@@ -0,0 +1,50 @@
+﻿// Copyright Epic Games, Inc. All Rights Reserved.
+
+#include "SandboxFileUtils.h"
+
+#include "JsonObjectConverter.h"
+#include "LogFileSandbox.h"
+#include "Data/ManifestData.h"
+#include "Misc/FileHelper.h"
+#include "Misc/Optional.h"
+#include "Misc/Paths.h"
+
+namespace UE::FileSandboxCore
+{
+bool SaveManifest(const FFileSandboxCore_ManifestData& InData, const FString& InRootDirectory)
+{
+	FString JSONPayload;
+	FJsonObjectConverter::UStructToJsonObjectString(InData, JSONPayload, 0, 0);
+	const FString AbsFileName = InRootDirectory / GetManifestFileName();
+	const bool bSuccess = FFileHelper::SaveStringToFile(JSONPayload, *AbsFileName);
+	UE_CLOG(!bSuccess, LogFileSandbox, Error, TEXT("Failed to write sandbox manifest file %s"), *InRootDirectory);
+	return bSuccess;
+}
+
+TOptional<FFileSandboxCore_ManifestData> LoadManifest(const FString& InRootDirectory, bool bLogErrorIfNotExist)
+{
+	FFileSandboxCore_ManifestData ManifestPlatformFileState;
+	FString JSONPayload;
+	
+	const FString AbsFileName = InRootDirectory / GetManifestFileName();
+	if (FFileHelper::LoadFileToString(JSONPayload, *AbsFileName))
+	{
+		const bool bSuccess = FJsonObjectConverter::JsonObjectStringToUStruct(JSONPayload, &ManifestPlatformFileState, 0, 0);
+		UE_CLOG(!bSuccess, LogFileSandbox, Error, TEXT("Failed to parse sandbox manifest file %s"), *InRootDirectory);
+		return bSuccess ? ManifestPlatformFileState : TOptional<FFileSandboxCore_ManifestData>{};
+	}
+	
+	UE_CLOG(bLogErrorIfNotExist, LogFileSandbox, Error, TEXT("Failed to load sandbox manifest file %s"), *InRootDirectory);
+	return {};
+}
+
+FString GetManifestFileName()
+{
+	return TEXT("manifest.json");
+}
+
+FString GetMetadataFileName()
+{
+	return TEXT("metadata.json");
+}
+}
diff --git a/Engine/Plugins/Developer/Sandbox/FileSandbox/Source/FileSandboxCore/Private/Utils/SandboxFileUtils.h b/Engine/Plugins/Developer/Sandbox/FileSandbox/Source/FileSandboxCore/Private/Utils/SandboxFileUtils.h
new file mode 100644
--- /dev/null
+++ b/Engine/Plugins/Developer/Sandbox/FileSandbox/Source/FileSandboxCore/Private/Utils/SandboxFileUtils.h
@@ -0,0 +1,20 @@
+﻿// Copyright Epic Games, Inc. All Rights Reserved.
+
+#pragma once
+
+class FString;
+struct FFileSandboxCore_ManifestData;
+template<typename OptionalType> struct TOptional;
+
+namespace UE::FileSandboxCore
+{
+/** Saves the manifest file. */
+bool SaveManifest(const FFileSandboxCore_ManifestData& InData, const FString& InRootDirectory);
+/** Loads the manifest file. */
+TOptional<FFileSandboxCore_ManifestData> LoadManifest(const FString& InRootDirectory, bool bLogErrorIfNotExist = true);
+
+/* @return The name to give the manifest file. */
+FString GetManifestFileName();
+/* @return The name to give the metadata file. */
+FString GetMetadataFileName();
+}
\ No newline at end of file
diff --git a/Engine/Plugins/Developer/Sandbox/FileSandbox/Source/FileSandboxCore/Private/Utils/SandboxUtils.cpp b/Engine/Plugins/Developer/Sandbox/FileSandbox/Source/FileSandboxCore/Private/Utils/SandboxUtils.cpp
deleted file mode 100644
diff --git a/Engine/Plugins/Developer/Sandbox/FileSandbox/Source/FileSandboxCore/Private/Utils/SandboxUtils.h b/Engine/Plugins/Developer/Sandbox/FileSandbox/Source/FileSandboxCore/Private/Utils/SandboxUtils.h
deleted file mode 100644
diff --git a/Engine/Plugins/Developer/Sandbox/FileSandbox/Source/FileSandboxCore/Private/Utils/ScopedWatchedDirectory.cpp b/Engine/Plugins/Developer/Sandbox/FileSandbox/Source/FileSandboxCore/Private/Utils/ScopedWatchedDirectory.cpp
deleted file mode 100644
diff --git a/Engine/Plugins/Developer/Sandbox/FileSandbox/Source/FileSandboxCore/Private/Utils/ScopedWatchedDirectory.h b/Engine/Plugins/Developer/Sandbox/FileSandbox/Source/FileSandboxCore/Private/Utils/ScopedWatchedDirectory.h
deleted file mode 100644
diff --git a/Engine/Plugins/Developer/Sandbox/FileSandbox/Source/FileSandboxCore/Private/Utils/Watcher/DirectoryWatcherUtils.cpp b/Engine/Plugins/Developer/Sandbox/FileSandbox/Source/FileSandboxCore/Private/Utils/Watcher/DirectoryWatcherUtils.cpp
new file mode 100644
--- /dev/null
+++ b/Engine/Plugins/Developer/Sandbox/FileSandbox/Source/FileSandboxCore/Private/Utils/Watcher/DirectoryWatcherUtils.cpp
@@ -0,0 +1,44 @@
+﻿// Copyright Epic Games, Inc. All Rights Reserved.
+
+#if UE_SANDBOX_WITH_DIRECTORY_WATCHER
+
+#include "DirectoryWatcherUtils.h"
+
+#include "DirectoryWatcherModule.h"
+#include "Modules/ModuleManager.h"
+#include "UObject/NameTypes.h"
+
+namespace UE::FileSandboxCore
+{
+FDirectoryWatcherModule& GetDirectoryWatcherModule()
+{
+	const FName DirectoryWatcherModuleName = TEXT("DirectoryWatcher");
+	return FModuleManager::LoadModuleChecked<FDirectoryWatcherModule>(DirectoryWatcherModuleName);
+}
+
+FDirectoryWatcherModule* GetDirectoryWatcherModuleIfLoaded()
+{
+	const FName DirectoryWatcherModuleName = TEXT("DirectoryWatcher");
+	if (FModuleManager::Get().IsModuleLoaded(DirectoryWatcherModuleName))
+	{
+		return &FModuleManager::GetModuleChecked<FDirectoryWatcherModule>(DirectoryWatcherModuleName);
+	}
+	return nullptr;
+}
+
+IDirectoryWatcher* GetDirectoryWatcher()
+{
+	FDirectoryWatcherModule& DirectoryWatcherModule = GetDirectoryWatcherModule();
+	return DirectoryWatcherModule.Get();
+}
+
+IDirectoryWatcher* GetDirectoryWatcherIfLoaded()
+{
+	if (FDirectoryWatcherModule* DirectoryWatcherModule = GetDirectoryWatcherModuleIfLoaded())
+	{
+		return DirectoryWatcherModule->Get();
+	}
+	return nullptr;
+}
+}
+#endif
\ No newline at end of file
diff --git a/Engine/Plugins/Developer/Sandbox/FileSandbox/Source/FileSandboxCore/Private/Utils/Watcher/DirectoryWatcherUtils.h b/Engine/Plugins/Developer/Sandbox/FileSandbox/Source/FileSandboxCore/Private/Utils/Watcher/DirectoryWatcherUtils.h
new file mode 100644
--- /dev/null
+++ b/Engine/Plugins/Developer/Sandbox/FileSandbox/Source/FileSandboxCore/Private/Utils/Watcher/DirectoryWatcherUtils.h
@@ -0,0 +1,18 @@
+﻿// Copyright Epic Games, Inc. All Rights Reserved.
+
+#pragma once
+
+#if UE_SANDBOX_WITH_DIRECTORY_WATCHER
+
+class IDirectoryWatcher;
+class FDirectoryWatcherModule;
+
+namespace UE::FileSandboxCore
+{
+FDirectoryWatcherModule& GetDirectoryWatcherModule();
+FDirectoryWatcherModule* GetDirectoryWatcherModuleIfLoaded();
+IDirectoryWatcher* GetDirectoryWatcher();
+IDirectoryWatcher* GetDirectoryWatcherIfLoaded();
+}
+
+#endif
\ No newline at end of file
diff --git a/Engine/Plugins/Developer/Sandbox/FileSandbox/Source/FileSandboxCore/Private/Utils/Watcher/ScopedWatchedDirectory.cpp b/Engine/Plugins/Developer/Sandbox/FileSandbox/Source/FileSandboxCore/Private/Utils/Watcher/ScopedWatchedDirectory.cpp
new file mode 100644
--- /dev/null
+++ b/Engine/Plugins/Developer/Sandbox/FileSandbox/Source/FileSandboxCore/Private/Utils/Watcher/ScopedWatchedDirectory.cpp
@@ -0,0 +1,37 @@
+﻿// Copyright Epic Games, Inc. All Rights Reserved.
+
+#if UE_SANDBOX_WITH_DIRECTORY_WATCHER
+
+#include "ScopedWatchedDirectory.h"
+
+#include "DirectoryWatcherUtils.h"
+#include "IDirectoryWatcher.h"
+
+namespace UE::FileSandboxCore
+{
+FScopedWatchedDirectory::FScopedWatchedDirectory(FString InDirectory, FDelegateHandle InHandle)
+	: WatchedDirectory(InDirectory)
+	, Handle(InHandle)
+{}
+
+FScopedWatchedDirectory::FScopedWatchedDirectory(FScopedWatchedDirectory&& Old)
+	: FScopedWatchedDirectory(MoveTemp(Old.WatchedDirectory), Old.Handle)
+{
+	Old.Handle.Reset();
+}
+
+FScopedWatchedDirectory::~FScopedWatchedDirectory()
+{
+	if (!Handle.IsValid() || WatchedDirectory.IsEmpty())
+	{
+		return;
+	}
+	
+	if (IDirectoryWatcher* Watcher = GetDirectoryWatcherIfLoaded())
+	{
+		Watcher->UnregisterDirectoryChangedCallback_Handle(WatchedDirectory, Handle);
+	}
+}
+}
+
+#endif
\ No newline at end of file
diff --git a/Engine/Plugins/Developer/Sandbox/FileSandbox/Source/FileSandboxCore/Private/Utils/Watcher/ScopedWatchedDirectory.h b/Engine/Plugins/Developer/Sandbox/FileSandbox/Source/FileSandboxCore/Private/Utils/Watcher/ScopedWatchedDirectory.h
new file mode 100644
--- /dev/null
+++ b/Engine/Plugins/Developer/Sandbox/FileSandbox/Source/FileSandboxCore/Private/Utils/Watcher/ScopedWatchedDirectory.h
@@ -0,0 +1,32 @@
+﻿// Copyright Epic Games, Inc. All Rights Reserved.
+
+#pragma once
+
+#if UE_SANDBOX_WITH_DIRECTORY_WATCHER
+
+#include "Containers/UnrealString.h"
+#include "Delegates/IDelegateInstance.h"
+#include "Templates/UnrealTemplate.h"
+
+namespace UE::FileSandboxCore
+{
+/** Unregisters IDirectoryWatcher when it goes out of scope. */
+class FScopedWatchedDirectory : public FNoncopyable
+{
+public:
+	
+	FScopedWatchedDirectory() = default;
+	FScopedWatchedDirectory(FString InDirectory, FDelegateHandle InHandle);
+	FScopedWatchedDirectory(FScopedWatchedDirectory&& Old);
+	~FScopedWatchedDirectory();
+	
+	const FString& GetWatchedDirectory() const { return WatchedDirectory; }
+
+private:
+	
+	FString WatchedDirectory;
+	FDelegateHandle Handle;
+};
+}
+
+#endif
\ No newline at end of file
diff --git a/Engine/Plugins/Developer/Sandbox/FileSandbox/Source/FileSandboxCore/Public/Utils/PackageSandboxUtils.h b/Engine/Plugins/Developer/Sandbox/FileSandbox/Source/FileSandboxCore/Public/Utils/PackageSandboxUtils.h
--- a/Engine/Plugins/Developer/Sandbox/FileSandbox/Source/FileSandboxCore/Public/Utils/PackageSandboxUtils.h
+++ b/Engine/Plugins/Developer/Sandbox/FileSandbox/Source/FileSandboxCore/Public/Utils/PackageSandboxUtils.h
@@ -8,7 +8,9 @@ class FName;
 
 namespace UE::FileSandboxCore
 {
+/** Handles cleaning up assets that no longer exist. */
 void PurgePackages(const TArray<FName>& InPackageNames);
 
+/** Reloads packages that were modified */
 void HotReloadPackages(const TArray<FName>& InPackageNames);
 }
