File Sandbox Core: Centralize analysis of removed / edited / added files

#Synced-CL 47637010

diff --git a/Engine/Plugins/Developer/Sandbox/FileSandbox/Source/FileSandboxCore/Private/Sandbox/Persist/DiffUtils.cpp b/Engine/Plugins/Developer/Sandbox/FileSandbox/Source/FileSandboxCore/Private/Sandbox/Persist/DiffUtils.cpp
new file mode 100644
--- /dev/null
+++ b/Engine/Plugins/Developer/Sandbox/FileSandbox/Source/FileSandboxCore/Private/Sandbox/Persist/DiffUtils.cpp
@@ -0,0 +1,65 @@
+﻿// Copyright Epic Games, Inc. All Rights Reserved.
+
+#include "DiffUtils.h"
+
+#include "Containers/ArrayView.h"
+#include "Containers/Set.h"
+#include "GenericPlatform/GenericPlatformFile.h"
+#include "Misc/Paths.h"
+#include "Sandbox/Platform/SandboxMountPoint.h"
+#include "Templates/Function.h"
+
+namespace UE::FileSandboxCore
+{
+void AnalyzeFilesMarkedRemoved(
+	IPlatformFile& InPlatformFile,
+	const TSet<FNonSandboxPath>& InFilesMarkedForRemove, 
+	TFunctionRef<void(const FNonSandboxPath&)> InHandleRemovedFile)
+{
+	for (const FNonSandboxPath& PathMarkedRemoved : InFilesMarkedForRemove)
+	{
+		if (InPlatformFile.FileExists(*PathMarkedRemoved))
+		{
+			InHandleRemovedFile(PathMarkedRemoved);
+		}
+	}
+}
+
+void ScanForEditedOrAddedFiles(
+	IPlatformFile& InPlatformFile,
+	TConstArrayView<FSandboxMountPoint> InSandboxMountPoints, 
+	TFunctionRef<FHandleFilePathSignature> InHandleEditedFile,
+	TFunctionRef<FHandleFilePathSignature> InHandleAddedFile
+	)
+{
+	for (const FSandboxMountPoint& MountPoint : InSandboxMountPoints)
+	{
+		InPlatformFile.IterateDirectoryRecursively(*MountPoint.Path.GetSandboxPath(), 
+			[&InPlatformFile, &MountPoint, &InHandleEditedFile, &InHandleAddedFile]
+			(const TCHAR* InFilenameOrDirectory, const bool InIsDirectory) -> bool
+		{
+			if (InIsDirectory)
+			{
+				return true;
+			}
+
+			const FSandboxedPlatformFilePath Path = FSandboxedPlatformFilePath::CreateNonSandboxPath(
+				FPaths::ConvertRelativePathToFull(InFilenameOrDirectory), MountPoint.Path
+				);
+			const FString& NonSandboxPath = Path.GetNonSandboxPath();
+					
+			const bool bWasEdited = InPlatformFile.FileExists(*NonSandboxPath);
+			if (bWasEdited)
+			{
+				InHandleEditedFile(Path);
+			}
+			else
+			{
+				InHandleAddedFile(Path);
+			}
+					
+			return true; // continue iteration
+		});
+	}
+}
+}
diff --git a/Engine/Plugins/Developer/Sandbox/FileSandbox/Source/FileSandboxCore/Private/Sandbox/Persist/DiffUtils.h b/Engine/Plugins/Developer/Sandbox/FileSandbox/Source/FileSandboxCore/Private/Sandbox/Persist/DiffUtils.h
new file mode 100644
--- /dev/null
+++ b/Engine/Plugins/Developer/Sandbox/FileSandbox/Source/FileSandboxCore/Private/Sandbox/Persist/DiffUtils.h
@@ -0,0 +1,48 @@
+﻿// Copyright Epic Games, Inc. All Rights Reserved.
+
+#pragma once
+
+#include "Containers/ContainersFwd.h"
+#include "Sandbox/Platform/SandboxedPlatformFilePath.h"
+#include "Templates/FunctionFwd.h"
+
+namespace UE::FileSandboxCore
+{
+struct FSandboxMountPoint;
+
+using FHandleFilePathSignature = void(const FSandboxedPlatformFilePath& InPath);
+
+/** 
+ * Analyzes the file action the sandbox system would perform on files previously marked for remove.
+ * This compares whether the non-sandbox path still exists, i.e. whether the file was already removed (externally).
+ * 
+ * This is useful if you're the state from the manifest file and need to compare against the current disk state.
+ * 
+ * @param InPlatformFile The lower level platform file, i.e. the one that FSandboxPlatformFile wraps.
+ * @param InFilesMarkedForRemove The files to analyze.
+ * @param InHandleRemovedFile Invoked for each file that still should be marked as removed.
+ */
+void AnalyzeFilesMarkedRemoved(
+	IPlatformFile& InPlatformFile,
+	const TSet<FNonSandboxPath>& InFilesMarkedForRemove, 
+	TFunctionRef<void(const FNonSandboxPath&)> InHandleRemovedFile
+	);
+
+/** 
+ * Recursively iterates the files in the mount points and determines whether the present files would be
+ * - editing the non-sandbox counterpart, or
+ * - adding a new file to the non-sandbox.
+ * This is useful if you're the state from the manifest file and need to compare against the current disk state.
+ * 
+ * @param InPlatformFile The lower level platform file, i.e. the one that FSandboxPlatformFile wraps.
+ * @param InSandboxMountPoints The sandbox mountpoints
+ * @param InHandleEditedFile Invoked for each file that should be marked as edited.
+ * @param InHandleAddedFile Invoked for each file that should be marked as added.
+ */
+void ScanForEditedOrAddedFiles(
+	IPlatformFile& InPlatformFile,
+	TConstArrayView<FSandboxMountPoint> InSandboxMountPoints,
+	TFunctionRef<FHandleFilePathSignature> InHandleEditedFile,
+	TFunctionRef<FHandleFilePathSignature> InHandleAddedFile
+	);
+}
\ No newline at end of file
diff --git a/Engine/Plugins/Developer/Sandbox/FileSandbox/Source/FileSandboxCore/Private/Sandbox/Persist/DiscardUtils.cpp b/Engine/Plugins/Developer/Sandbox/FileSandbox/Source/FileSandboxCore/Private/Sandbox/Persist/DiscardUtils.cpp
--- a/Engine/Plugins/Developer/Sandbox/FileSandbox/Source/FileSandboxCore/Private/Sandbox/Persist/DiscardUtils.cpp
+++ b/Engine/Plugins/Developer/Sandbox/FileSandbox/Source/FileSandboxCore/Private/Sandbox/Persist/DiscardUtils.cpp
@@ -2,6 +2,7 @@
 
 #include "DiscardUtils.h"
 
+#include "DiffUtils.h"
 #include "Containers/Array.h"
 #include "Containers/ArrayView.h"
 #include "Containers/Set.h"
@@ -20,22 +21,20 @@ void DiscardRemovedPaths(
 	TFunctionRef<void(const FFileChange&)> InProcessFileChangeCallback
 	)
 {
-	for (const FSandboxPath& NonSandboxPath : InDeletedSandboxPaths)
-	{
-		// If this file maps to a package then we need to flush its linker so that we can remove the file from the sandbox
-		FName PackageName;
-		FlushPackageFile(NonSandboxPath, &PackageName);
-
-		if (InLowerLevel.FileExists(*NonSandboxPath))
+	AnalyzeFilesMarkedRemoved(InLowerLevel, InDeletedSandboxPaths,	
+		[&OutPackagesPendingPurge, &InProcessFileChangeCallback](const FNonSandboxPath& NonSandboxPath)
 		{
+			// If this file maps to a package then we need to flush its linker so that we can remove the file from the sandbox
+			FName PackageName;
+			FlushPackageFile(NonSandboxPath, &PackageName);
+				
 			if (!PackageName.IsNone())
 			{
 				OutPackagesPendingPurge.Add(PackageName);
 			}
-			
+		
 			InProcessFileChangeCallback(FFileChange{ NonSandboxPath, EFileChangeAction::Added });
-		}
-	}
+		});
 }
 
 void DiscardAddedAndModifiedPaths(
@@ -45,62 +44,53 @@ void DiscardAddedAndModifiedPaths(
 	TFunctionRef<void(const FFileChange&)> InProcessFileChangeCallback
 	)
 {
-	for (const FSandboxMountPoint& SandboxMountPoint : InMountPoints)
+	const auto HandleEditedFile = [&](const FSandboxedPlatformFilePath& InPath)
 	{
-		// Modify any files that exist in both the non-sandbox and sandbox directories
-		// Delete any files that exist in the sandbox but don't exist in the non-sandbox directory
-		InLowerLevel.IterateDirectoryRecursively(*SandboxMountPoint.Path.GetSandboxPath(), 
-			[&InLowerLevel, &SandboxMountPoint, &OutPackagesPendingHotReload, &OutPackagesPendingPurge, &InProcessFileChangeCallback]
-			(const TCHAR* InFilenameOrDirectory, const bool InIsDirectory) -> bool
+		// If this file maps to a package then we need to flush its linker so that we can remove the file from the sandbox
+		FName PackageName;
+		FlushPackageFile(InPath.GetNonSandboxPath(), &PackageName, false);
+			
+		if (!PackageName.IsNone())
 		{
-			if (!InIsDirectory)
-			{
-				const FSandboxedPlatformFilePath RemappedFilePath = FSandboxedPlatformFilePath::CreateNonSandboxPath(FPaths::ConvertRelativePathToFull(InFilenameOrDirectory), SandboxMountPoint.Path);
-
-				// If this file maps to a package then we need to flush its linker so that we can remove the file from the sandbox
-				FName PackageName;
-				FlushPackageFile(RemappedFilePath.GetNonSandboxPath(), &PackageName, false);
-
-				if (InLowerLevel.FileExists(*RemappedFilePath.GetNonSandboxPath()))
-				{
-					if (!PackageName.IsNone())
-					{
-						OutPackagesPendingHotReload.Add(PackageName);
-						OutPackagesPendingPurge.Remove(PackageName);
-					}
-					
-					InProcessFileChangeCallback(FFileChange{ RemappedFilePath.GetNonSandboxPath(), EFileChangeAction::Modified });
-				}
-				else
-				{
-					if (!PackageName.IsNone())
-					{
-						OutPackagesPendingPurge.Add(PackageName);
-						OutPackagesPendingHotReload.Remove(PackageName);
-					}
-					
-					InProcessFileChangeCallback(FFileChange{ RemappedFilePath.GetNonSandboxPath(), EFileChangeAction::Removed });
-				}
-			}
-			return true; // Continue iteration
-		});
+			OutPackagesPendingHotReload.Add(PackageName);
+			OutPackagesPendingPurge.Remove(PackageName);
+		}
+				
+		InProcessFileChangeCallback(FFileChange{ InPath.GetNonSandboxPath(), EFileChangeAction::Modified });
+	};
+	
+	const auto HandleAddedFile = [&](const FSandboxedPlatformFilePath& InPath)
+	{
+		// If this file maps to a package then we need to flush its linker so that we can remove the file from the sandbox
+		FName PackageName;
+		FlushPackageFile(InPath.GetNonSandboxPath(), &PackageName, false);
+			
+		if (!PackageName.IsNone())
+		{
+			OutPackagesPendingPurge.Add(PackageName);
+			OutPackagesPendingHotReload.Remove(PackageName);
+		}
+				
+		InProcessFileChangeCallback(FFileChange{ InPath.GetNonSandboxPath(), EFileChangeAction::Added });
+	};
+	
+	ScanForEditedOrAddedFiles(InLowerLevel, InMountPoints, HandleEditedFile, HandleAddedFile);
 
-		//FString SandboxRootPath = GetSandboxRootPath();
-		//FString Stashed = SandboxRootPath.Replace(TEXT("Sandbox"), TEXT("Stashed"));
-		//LowerLevel->MoveFile(SandboxRootPath, Stashed);
-		// Delete everything under the mount point
-		//LowerLevel->IterateDirectory(*SandboxMountPoint.Path.GetSandboxPath(), [this](const TCHAR* InFilenameOrDirectory, const bool InIsDirectory) -> bool
-		//{
-		//	if (InIsDirectory)
-		//	{
-		//		LowerLevel->DeleteDirectoryRecursively(InFilenameOrDirectory);
-		//	}
-		//	else
-		//	{
-		//		LowerLevel->DeleteFile(InFilenameOrDirectory);
-		//	}
-		//	return true; // Continue iteration
-		//});
-	}
+	//FString SandboxRootPath = GetSandboxRootPath();
+	//FString Stashed = SandboxRootPath.Replace(TEXT("Sandbox"), TEXT("Stashed"));
+	//LowerLevel->MoveFile(SandboxRootPath, Stashed);
+	// Delete everything under the mount point
+	//LowerLevel->IterateDirectory(*SandboxMountPoint.Path.GetSandboxPath(), [this](const TCHAR* InFilenameOrDirectory, const bool InIsDirectory) -> bool
+	//{
+	//	if (InIsDirectory)
+	//	{
+	//		LowerLevel->DeleteDirectoryRecursively(InFilenameOrDirectory);
+	//	}
+	//	else
+	//	{
+	//		LowerLevel->DeleteFile(InFilenameOrDirectory);
+	//	}
+	//	return true; // Continue iteration
+	//});
 }
 }
diff --git a/Engine/Plugins/Developer/Sandbox/FileSandbox/Source/FileSandboxCore/Private/Sandbox/Platform/SandboxPlatformFile.cpp b/Engine/Plugins/Developer/Sandbox/FileSandbox/Source/FileSandboxCore/Private/Sandbox/Platform/SandboxPlatformFile.cpp
--- a/Engine/Plugins/Developer/Sandbox/FileSandbox/Source/FileSandboxCore/Private/Sandbox/Platform/SandboxPlatformFile.cpp
+++ b/Engine/Plugins/Developer/Sandbox/FileSandbox/Source/FileSandboxCore/Private/Sandbox/Platform/SandboxPlatformFile.cpp
@@ -20,6 +20,7 @@
 #include "Sandbox/Persist/DiscardUtils.h"
 #include "Sandbox/Persist/PersistUtils.h"
 #include "SandboxMountPoint.h"
+#include "Sandbox/Persist/DiffUtils.h"
 #include "Types/Sandbox/PersistArgs.h"
 #include "Types/Sandbox/PersistResult.h"
 #include "Types/SandboxFileChange.h"
@@ -790,36 +791,20 @@ TArray<FString> FSandboxPlatformFile::GatherSandboxChangedFilenames() const
 	{
 		const FScopeLock DeletedSandboxPathsLock(&DeletedSandboxPathsCS);
 		ChangedFiles.Reserve(DeletedSandboxPaths.Num());
-		Algo::TransformIf(DeletedSandboxPaths, ChangedFiles,
-			[this](const FNonSandboxPath& DeletedPath)
-			{
-				return LowerLevel->FileExists(*DeletedPath);
-			},
-			[this](const FNonSandboxPath& DeletedPath)
-			{
-				return DeletedPath;
-			});
+		AnalyzeFilesMarkedRemoved(*LowerLevel, DeletedSandboxPaths, [&ChangedFiles](const FNonSandboxPath& RemovedNonSandboxPath)
+		{
+			ChangedFiles.Add(RemovedNonSandboxPath);
+		});
 	}
 
 	// Gather mounts files
 	{
 		const FScopeLock SandboxMountPointsLock(&SandboxMountPointsCS);
-		for (const FSandboxMountPoint& MountPoint : SandboxMountPoints)
+		const auto HandleChangedFile = [&ChangedFiles](const FSandboxedPlatformFilePath& InPath)
 		{
-			LowerLevel->IterateDirectoryRecursively(*MountPoint.Path.GetSandboxPath(), [this, &MountPoint, &ChangedFiles](const TCHAR* InFilenameOrDirectory, const bool InIsDirectory) -> bool
-			{
-				if (InIsDirectory)
-				{
-					return true;
-				}
-
-				const FSandboxedPlatformFilePath Path = FSandboxedPlatformFilePath::CreateNonSandboxPath(
-					FPaths::ConvertRelativePathToFull(InFilenameOrDirectory), MountPoint.Path
-					);
-				ChangedFiles.Add(Path.GetNonSandboxPath());
-				return true; // continue iteration
-			});
-		}
+			ChangedFiles.Add(InPath.GetNonSandboxPath());
+		};
+		ScanForEditedOrAddedFiles(*LowerLevel, SandboxMountPoints, HandleChangedFile, HandleChangedFile);
 	}
 
 	return ChangedFiles;
diff --git a/Engine/Plugins/Developer/Sandbox/FileSandbox/Source/FileSandboxCore/Private/Sandbox/Platform/SandboxPlatformFile.h b/Engine/Plugins/Developer/Sandbox/FileSandbox/Source/FileSandboxCore/Private/Sandbox/Platform/SandboxPlatformFile.h
--- a/Engine/Plugins/Developer/Sandbox/FileSandbox/Source/FileSandboxCore/Private/Sandbox/Platform/SandboxPlatformFile.h
+++ b/Engine/Plugins/Developer/Sandbox/FileSandbox/Source/FileSandboxCore/Private/Sandbox/Platform/SandboxPlatformFile.h
@@ -189,7 +189,7 @@ private:
 	 * 
 	 * This contains the absolute path to the non-sandbox file or directory.
 	 */
-	TSet<FSandboxPath> DeletedSandboxPaths;
+	TSet<FNonSandboxPath> DeletedSandboxPaths;
 	/** Critical section protecting concurrent access to DeletedSandboxPaths */
 	mutable FCriticalSection DeletedSandboxPathsCS;
 	
