File Sandbox: Add API for getting what has changed about files.

#Synced-CL 49159615

diff --git a/Engine/Plugins/Developer/Sandbox/FileSandbox/Source/FileSandboxCore/Private/ISandboxManager.cpp b/Engine/Plugins/Developer/Sandbox/FileSandbox/Source/FileSandboxCore/Private/ISandboxManager.cpp
--- a/Engine/Plugins/Developer/Sandbox/FileSandbox/Source/FileSandboxCore/Private/ISandboxManager.cpp
+++ b/Engine/Plugins/Developer/Sandbox/FileSandbox/Source/FileSandboxCore/Private/ISandboxManager.cpp
@@ -48,4 +48,14 @@ bool ISandboxManager::CanLeaveSandbox() const
 	ISandboxLock* ActiveLock = GetActiveLock();
 	return !ActiveLock || ActiveLock->CanLeaveSandbox();
 }
+
+void ISandboxManager::EnumerateFileChangesByName(
+	const FString& InSandboxName, TFunctionRef<FProcessFileChangeSignature> InProcess, const FString& InBaseDirectory
+	) const
+{
+	if (const TOptional<FString> FullPath = FindSandboxByName(InSandboxName, InBaseDirectory))
+	{
+		EnumerateFileChanges(*FullPath, InProcess);
+	}
+}
 }
diff --git a/Engine/Plugins/Developer/Sandbox/FileSandbox/Source/FileSandboxCore/Private/Sandbox/ISandboxInstance.cpp b/Engine/Plugins/Developer/Sandbox/FileSandbox/Source/FileSandboxCore/Private/Sandbox/ISandboxInstance.cpp
--- a/Engine/Plugins/Developer/Sandbox/FileSandbox/Source/FileSandboxCore/Private/Sandbox/ISandboxInstance.cpp
+++ b/Engine/Plugins/Developer/Sandbox/FileSandbox/Source/FileSandboxCore/Private/Sandbox/ISandboxInstance.cpp
@@ -2,6 +2,7 @@
 
 #include "ISandboxInstance.h"
 
+#include "Types/EBreakBehavior.h"
 #include "Types/Sandbox/DiscardArgs.h"
 #include "Types/Sandbox/DiscardResult.h"
 #include "Types/Sandbox/PersistArgs.h"
@@ -11,7 +12,8 @@ namespace UE::FileSandboxCore
 {
 void ISandboxInstance::PersistAll()
 {
-	const TArray<FString> ChangedFiles = GatherChangedFiles();
+	const FGatheredFileChanges FileChanges = GatherChangedFiles(EFileChangeGatherFlags::None);
+	const TArray<FString>& ChangedFiles = FileChanges.NonSandboxPaths;
 	PersistSandbox(FPersistArgs{ ChangedFiles });
 }
 
@@ -21,4 +23,19 @@ void ISandboxInstance::DiscardAll()
 	HotReloadPackages(Result.PackagesPendingHotReload);
 	PurgePackages(Result.PackagesPendingPurge);
 }
+
+FGatheredFileChanges ISandboxInstance::GatherChangedFiles(EFileChangeGatherFlags InFlags) const
+{
+	FGatheredFileChanges Result;
+	EnumerateFileChanges([InFlags, &Result](const FString& InFile, ESandboxFileChange InChange)
+	{
+		Result.NonSandboxPaths.Add(InFile);
+		if (EnumHasAnyFlags(InFlags, EFileChangeGatherFlags::IncludeChangeTypes))
+		{
+			Result.FileActions.Add(InChange);
+		}
+		return EBreakBehavior::Continue;
+	});
+	return Result;
+}
 }
diff --git a/Engine/Plugins/Developer/Sandbox/FileSandbox/Source/FileSandboxCore/Private/Sandbox/Persist/DiffUtils.cpp b/Engine/Plugins/Developer/Sandbox/FileSandbox/Source/FileSandboxCore/Private/Sandbox/Persist/DiffUtils.cpp
--- a/Engine/Plugins/Developer/Sandbox/FileSandbox/Source/FileSandboxCore/Private/Sandbox/Persist/DiffUtils.cpp
+++ b/Engine/Plugins/Developer/Sandbox/FileSandbox/Source/FileSandboxCore/Private/Sandbox/Persist/DiffUtils.cpp
@@ -4,39 +4,64 @@
 
 #include "Containers/ArrayView.h"
 #include "Containers/Set.h"
+#include "Data/ManifestData.h"
 #include "GenericPlatform/GenericPlatformFile.h"
+#include "Misc/PackageName.h"
 #include "Misc/Paths.h"
 #include "Sandbox/Platform/SandboxMountPoint.h"
 #include "Templates/Function.h"
+#include "Types/EBreakBehavior.h"
+#include "Types/SandboxFileChange.h"
+#include "Utils/SandboxFileUtils.h"
 
 namespace UE::FileSandboxCore
 {
-void AnalyzeFilesMarkedRemoved(
-	IPlatformFile& InPlatformFile,
-	const TSet<FNonSandboxPath>& InFilesMarkedForRemove, 
-	TFunctionRef<void(const FNonSandboxPath&)> InHandleRemovedFile)
+namespace DiffDetail
+{
+template<typename TRangeType>
+static void AnalyzeFilesMarkedRemoved(
+	IPlatformFile& InPlatformFile, const TRangeType& InFilesMarkedForRemove, 
+	TFunctionRef<EBreakBehavior(const FNonSandboxPath&)> InHandleRemovedFile
+	)
 {
 	for (const FNonSandboxPath& PathMarkedRemoved : InFilesMarkedForRemove)
 	{
-		if (InPlatformFile.FileExists(*PathMarkedRemoved))
+		if (InPlatformFile.FileExists(*PathMarkedRemoved)
+			&& InHandleRemovedFile(PathMarkedRemoved) == EBreakBehavior::Break)
 		{
-			InHandleRemovedFile(PathMarkedRemoved);
+			return;
 		}
 	}
 }
+}
+
+void AnalyzeFilesMarkedRemoved(
+	IPlatformFile& InPlatformFile,
+	const TSet<FNonSandboxPath>& InFilesMarkedForRemove, 
+	TFunctionRef<EBreakBehavior(const FNonSandboxPath&)> InHandleRemovedFile
+	)
+{
+	DiffDetail::AnalyzeFilesMarkedRemoved(InPlatformFile, InFilesMarkedForRemove, InHandleRemovedFile);
+}
+
+void AnalyzeFilesMarkedRemoved(
+	IPlatformFile& InPlatformFile, const TConstArrayView<FNonSandboxPath>& InFilesMarkedForRemove,
+	TFunctionRef<EBreakBehavior(const FNonSandboxPath&)> InHandleRemovedFile
+	)
+{
+	DiffDetail::AnalyzeFilesMarkedRemoved(InPlatformFile, InFilesMarkedForRemove, InHandleRemovedFile);
+}
 
 void ScanForEditedOrAddedFiles(
 	IPlatformFile& InPlatformFile,
-	TConstArrayView<FSandboxMountPoint> InSandboxMountPoints, 
-	TFunctionRef<FHandleFilePathSignature> InHandleEditedFile,
+	const FSandboxedPlatformFilePath& InSandboxMountPoint,
+	TFunctionRef<FHandleFilePathSignature> InHandleEditedFile, 
 	TFunctionRef<FHandleFilePathSignature> InHandleAddedFile
 	)
 {
-	for (const FSandboxMountPoint& MountPoint : InSandboxMountPoints)
-	{
-		InPlatformFile.IterateDirectoryRecursively(*MountPoint.Path.GetSandboxPath(), 
-			[&InPlatformFile, &MountPoint, &InHandleEditedFile, &InHandleAddedFile]
-			(const TCHAR* InFilenameOrDirectory, const bool InIsDirectory) -> bool
+	InPlatformFile.IterateDirectoryRecursively(*InSandboxMountPoint.GetSandboxPath(), 
+		[&InPlatformFile, &InSandboxMountPoint, &InHandleEditedFile, &InHandleAddedFile]
+		(const TCHAR* InFilenameOrDirectory, const bool InIsDirectory) -> bool
 		{
 			if (InIsDirectory)
 			{
@@ -44,22 +69,90 @@ void ScanForEditedOrAddedFiles(
 			}
 
 			const FSandboxedPlatformFilePath Path = FSandboxedPlatformFilePath::CreateNonSandboxPath(
-				FPaths::ConvertRelativePathToFull(InFilenameOrDirectory), MountPoint.Path
+				FPaths::ConvertRelativePathToFull(InFilenameOrDirectory), InSandboxMountPoint
 				);
 			const FString& NonSandboxPath = Path.GetNonSandboxPath();
 					
 			const bool bWasEdited = InPlatformFile.FileExists(*NonSandboxPath);
 			if (bWasEdited)
 			{
-				InHandleEditedFile(Path);
+				return InHandleEditedFile(Path) == EBreakBehavior::Continue;
 			}
 			else
 			{
-				InHandleAddedFile(Path);
+				return InHandleAddedFile(Path) == EBreakBehavior::Continue;
 			}
 					
 			return true; // continue iteration
 		});
+}
+
+void ScanForEditedOrAddedFiles(
+	IPlatformFile& InPlatformFile,
+	TConstArrayView<FSandboxMountPoint> InSandboxMountPoints, 
+	TFunctionRef<FHandleFilePathSignature> InHandleEditedFile,
+	TFunctionRef<FHandleFilePathSignature> InHandleAddedFile
+	)
+{
+	for (const FSandboxMountPoint& MountPoint : InSandboxMountPoints)
+	{
+		ScanForEditedOrAddedFiles(InPlatformFile, MountPoint.Path, InHandleEditedFile, InHandleAddedFile);
 	}
 }
+
+void EnumerateMountPoints(TFunctionRef<EBreakBehavior(const FString& InAssetPath, const FString& InFilesystemPath)> InProcess)
+{
+	TArray<FString> RootPaths;
+	FPackageName::QueryRootContentPaths(RootPaths);
+	for (const FString& AssetPath : RootPaths)
+	{
+		const FString FilesystemPath = FPackageName::LongPackageNameToFilename(AssetPath);
+		if (InProcess(AssetPath, FilesystemPath) == EBreakBehavior::Break)
+		{
+			break;
+		}
+	}	
+}
+
+void EnumerateFileChanges(
+	IPlatformFile& InPlatformFile,
+	const FString& InSandboxRootPath, 
+	const FFileSandboxCore_ManifestData& InManifest,
+	TFunctionRef<FProcessFileChangeSignature> InProcess
+	)
+{
+	EBreakBehavior BreakBehavior = EBreakBehavior::Continue;
+	AnalyzeFilesMarkedRemoved(InPlatformFile, InManifest.DeletedFiles.GetFiles(), [&InProcess, &BreakBehavior](const FString& InFile)
+	{
+		BreakBehavior = InProcess(InFile, ESandboxFileChange::Removed);
+		return BreakBehavior;
+	});
+	if (BreakBehavior == EBreakBehavior::Break)
+	{
+		return;
+	}
+	
+	const auto ProcessMountPoint = [&InSandboxRootPath, &InProcess, &InPlatformFile, &BreakBehavior]
+		(const FString& InAssetPath, const FString& InFilesystemPath)
+	{
+		const FSandboxedPlatformFilePath MountPoint = FSandboxedPlatformFilePath::CreateMountPoint(
+			GetSandboxMountPointRoot(InSandboxRootPath), InAssetPath, InFilesystemPath
+		);
+			
+		const auto HandleEdited = [&InProcess, &BreakBehavior](const FSandboxedPlatformFilePath& InPath)
+		{
+			BreakBehavior = InProcess(InPath.GetNonSandboxPath(), ESandboxFileChange::Edited);
+			return BreakBehavior;
+		};
+		const auto HandleAdded = [&InProcess, &BreakBehavior](const FSandboxedPlatformFilePath& InPath)
+		{
+			BreakBehavior = InProcess(InPath.GetNonSandboxPath(), ESandboxFileChange::Added);
+			return BreakBehavior;
+		};
+			
+		ScanForEditedOrAddedFiles(InPlatformFile, MountPoint, HandleEdited, HandleAdded);
+		return BreakBehavior;
+	};
+	EnumerateMountPoints(ProcessMountPoint);
+}
 }
diff --git a/Engine/Plugins/Developer/Sandbox/FileSandbox/Source/FileSandboxCore/Private/Sandbox/Persist/DiffUtils.h b/Engine/Plugins/Developer/Sandbox/FileSandbox/Source/FileSandboxCore/Private/Sandbox/Persist/DiffUtils.h
--- a/Engine/Plugins/Developer/Sandbox/FileSandbox/Source/FileSandboxCore/Private/Sandbox/Persist/DiffUtils.h
+++ b/Engine/Plugins/Developer/Sandbox/FileSandbox/Source/FileSandboxCore/Private/Sandbox/Persist/DiffUtils.h
@@ -2,15 +2,20 @@
 
 #pragma once
 
+#include "ISandboxManager.h"
 #include "Containers/ContainersFwd.h"
 #include "Sandbox/Platform/SandboxedPlatformFilePath.h"
 #include "Templates/FunctionFwd.h"
+#include "Types/EBreakBehavior.h"
+
+struct FFileSandboxCore_ManifestData;
 
 namespace UE::FileSandboxCore
 {
+enum class EBreakBehavior : uint8;
 struct FSandboxMountPoint;
 
-using FHandleFilePathSignature = void(const FSandboxedPlatformFilePath& InPath);
+using FHandleFilePathSignature = EBreakBehavior(const FSandboxedPlatformFilePath& InPath);
 
 /** 
  * Analyzes the file action the sandbox system would perform on files previously marked for remove.
@@ -25,24 +30,50 @@ using FHandleFilePathSignature = void(const FSandboxedPlatformFilePath& InPath);
 void AnalyzeFilesMarkedRemoved(
 	IPlatformFile& InPlatformFile,
 	const TSet<FNonSandboxPath>& InFilesMarkedForRemove, 
-	TFunctionRef<void(const FNonSandboxPath&)> InHandleRemovedFile
+	TFunctionRef<EBreakBehavior(const FNonSandboxPath&)> InHandleRemovedFile
+	);
+void AnalyzeFilesMarkedRemoved(
+	IPlatformFile& InPlatformFile,
+	const TConstArrayView<FNonSandboxPath>& InFilesMarkedForRemove, 
+	TFunctionRef<EBreakBehavior(const FNonSandboxPath&)> InHandleRemovedFile
 	);
 
 /** 
- * Recursively iterates the files in the mount points and determines whether the present files would be
+ * Recursively iterates the files in the mount point and determines whether the present files would be
  * - editing the non-sandbox counterpart, or
  * - adding a new file to the non-sandbox.
  * This is useful if you're the state from the manifest file and need to compare against the current disk state.
  * 
  * @param InPlatformFile The lower level platform file, i.e. the one that FSandboxPlatformFile wraps.
- * @param InSandboxMountPoints The sandbox mountpoints
+ * @param InSandboxMountPoint The sandbox mountpoint
  * @param InHandleEditedFile Invoked for each file that should be marked as edited.
  * @param InHandleAddedFile Invoked for each file that should be marked as added.
  */
+void ScanForEditedOrAddedFiles(
+	IPlatformFile& InPlatformFile,
+	const FSandboxedPlatformFilePath& InSandboxMountPoint,
+	TFunctionRef<FHandleFilePathSignature> InHandleEditedFile,
+	TFunctionRef<FHandleFilePathSignature> InHandleAddedFile
+	);
+
+/** Utility overload for FSandboxPlatformFile to pass in all of its mount points. */
 void ScanForEditedOrAddedFiles(
 	IPlatformFile& InPlatformFile,
 	TConstArrayView<FSandboxMountPoint> InSandboxMountPoints,
 	TFunctionRef<FHandleFilePathSignature> InHandleEditedFile,
 	TFunctionRef<FHandleFilePathSignature> InHandleAddedFile
 	);
-}
\ No newline at end of file
+
+/** Enumerates all mount points in the engine. */
+void EnumerateMountPoints(TFunctionRef<EBreakBehavior(const FString& InAssetPath, const FString& InFilesystemPath)> InProcess);
+
+/** 
+ * Shared logic for determining the effective file changes.
+ * @see ISandboxInstance::EnumerateFileChanges.
+ * @see ISandboxManager::EnumerateFileChanges.
+ */
+void EnumerateFileChanges(
+	IPlatformFile& InPlatformFile, const FString& InSandboxRootPath, const FFileSandboxCore_ManifestData& InManifest, 
+	TFunctionRef<FProcessFileChangeSignature> InProcess
+	);
+}
diff --git a/Engine/Plugins/Developer/Sandbox/FileSandbox/Source/FileSandboxCore/Private/Sandbox/Persist/DiscardUtils.cpp b/Engine/Plugins/Developer/Sandbox/FileSandbox/Source/FileSandboxCore/Private/Sandbox/Persist/DiscardUtils.cpp
--- a/Engine/Plugins/Developer/Sandbox/FileSandbox/Source/FileSandboxCore/Private/Sandbox/Persist/DiscardUtils.cpp
+++ b/Engine/Plugins/Developer/Sandbox/FileSandbox/Source/FileSandboxCore/Private/Sandbox/Persist/DiscardUtils.cpp
@@ -34,6 +34,7 @@ void DiscardRemovedPaths(
 			}
 		
 			InProcessFileChangeCallback(FFileChange{ NonSandboxPath, EFileChangeAction::Added });
+			return EBreakBehavior::Continue;
 		});
 }
 
@@ -57,6 +58,7 @@ void DiscardAddedAndModifiedPaths(
 		}
 				
 		InProcessFileChangeCallback(FFileChange{ InPath.GetNonSandboxPath(), EFileChangeAction::Modified });
+		return EBreakBehavior::Continue;
 	};
 	
 	const auto HandleAddedFile = [&](const FSandboxedPlatformFilePath& InPath)
@@ -72,6 +74,7 @@ void DiscardAddedAndModifiedPaths(
 		}
 				
 		InProcessFileChangeCallback(FFileChange{ InPath.GetNonSandboxPath(), EFileChangeAction::Added });
+		return EBreakBehavior::Continue;
 	};
 	
 	ScanForEditedOrAddedFiles(InLowerLevel, InMountPoints, HandleEditedFile, HandleAddedFile);
diff --git a/Engine/Plugins/Developer/Sandbox/FileSandbox/Source/FileSandboxCore/Private/Sandbox/Platform/SandboxPlatformFile.cpp b/Engine/Plugins/Developer/Sandbox/FileSandbox/Source/FileSandboxCore/Private/Sandbox/Platform/SandboxPlatformFile.cpp
--- a/Engine/Plugins/Developer/Sandbox/FileSandbox/Source/FileSandboxCore/Private/Sandbox/Platform/SandboxPlatformFile.cpp
+++ b/Engine/Plugins/Developer/Sandbox/FileSandbox/Source/FileSandboxCore/Private/Sandbox/Platform/SandboxPlatformFile.cpp
@@ -2,6 +2,7 @@
 
 #include "SandboxPlatformFile.h"
 
+#include "Algo/ForEach.h"
 #include "Containers/UnrealString.h"
 #include "Engine/GameEngine.h"
 #include "Engine/World.h"
@@ -21,12 +22,14 @@
 #include "Sandbox/Persist/PersistUtils.h"
 #include "SandboxMountPoint.h"
 #include "Sandbox/Persist/DiffUtils.h"
+#include "Types/EBreakBehavior.h"
 #include "Types/Sandbox/PersistArgs.h"
 #include "Types/Sandbox/PersistResult.h"
 #include "Types/SandboxFileChange.h"
 #include "UObject/LinkerLoad.h"
 #include "UObject/Package.h"
 #include "Utils/AssetUtils.h"
+#include "Utils/SandboxFileUtils.h"
 #include "Utils/FileChange/FileChange.h"
 
 #if UE_SANDBOX_WITH_DIRECTORY_WATCHER
@@ -81,7 +84,7 @@ FSandboxPlatformFile::~FSandboxPlatformFile()
 
 FString FSandboxPlatformFile::GetSandboxRootPath() const
 {
-	return RootPath / TEXT("Sandbox");
+	return GetSandboxMountPointRoot(RootPath);
 }
 
 void FSandboxPlatformFile::SetSandboxEnabled(bool bInEnabled)
@@ -101,12 +104,11 @@ bool FSandboxPlatformFile::Initialize(IPlatformFile* Inner, const TCHAR* CmdLine
 	LowerLevel->CreateDirectoryTree(*GetSandboxRootPath());
 
 	// Set-up the initial set of content mount paths
-	TArray<FString> RootPaths;
-	FPackageName::QueryRootContentPaths(RootPaths);
-	for (const FString& Path : RootPaths)
+	EnumerateMountPoints([this](const FString& InAssetPath, const FString& InFilesystemPath)
 	{
-		RegisterContentMountPath(Path, FPackageName::LongPackageNameToFilename(Path));
-	}
+		RegisterContentMountPath(InAssetPath, InFilesystemPath);
+		return EBreakBehavior::Continue;
+	});
 
 	// Watch for new content mount paths
 	FPackageName::OnContentPathMounted().AddRaw(this, &FSandboxPlatformFile::OnContentPathMounted);
@@ -783,33 +785,6 @@ void FSandboxPlatformFile::DiscardSandbox(
 #endif
 }
 
-TArray<FString> FSandboxPlatformFile::GatherSandboxChangedFilenames() const
-{
-	TArray<FString> ChangedFiles;
-	
-	// Gather deleted paths
-	{
-		const FScopeLock DeletedSandboxPathsLock(&DeletedSandboxPathsCS);
-		ChangedFiles.Reserve(DeletedSandboxPaths.Num());
-		AnalyzeFilesMarkedRemoved(*LowerLevel, DeletedSandboxPaths, [&ChangedFiles](const FNonSandboxPath& RemovedNonSandboxPath)
-		{
-			ChangedFiles.Add(RemovedNonSandboxPath);
-		});
-	}
-
-	// Gather mounts files
-	{
-		const FScopeLock SandboxMountPointsLock(&SandboxMountPointsCS);
-		const auto HandleChangedFile = [&ChangedFiles](const FSandboxedPlatformFilePath& InPath)
-		{
-			ChangedFiles.Add(InPath.GetNonSandboxPath());
-		};
-		ScanForEditedOrAddedFiles(*LowerLevel, SandboxMountPoints, HandleChangedFile, HandleChangedFile);
-	}
-
-	return ChangedFiles;
-}
-
 bool FSandboxPlatformFile::DeletedPackageExistsInNonSandbox(FString InFilename) const
 {
 	const FSandboxedPlatformFilePath ResolvedPath = ToSandboxPath(MoveTemp(InFilename));
@@ -830,22 +805,20 @@ void FSandboxPlatformFile::OnContentPathDismounted(const FString& InAssetPath, c
 
 void FSandboxPlatformFile::RegisterContentMountPath(const FString& InAssetPath, const FString& InFilesystemPath)
 {
-	FString AbsoluteSandboxPath = FPaths::ConvertRelativePathToFull(GetSandboxRootPath() / InAssetPath) / TEXT("");
-	FString AbsoluteNonSandboxPath = FPaths::ConvertRelativePathToFull(InFilesystemPath) / TEXT("");
-	const bool bCreatedDirectory = LowerLevel->CreateDirectory(*AbsoluteSandboxPath);
-	
+	const FSandboxedPlatformFilePath MountPointPath = FSandboxedPlatformFilePath::CreateMountPoint(
+		GetSandboxRootPath(), InAssetPath, InFilesystemPath
+		);
+	const bool bCreatedDirectory = LowerLevel->CreateDirectory(*MountPointPath.GetSandboxPath());
 	UE_CLOG(!bCreatedDirectory, LogFileSandbox, Error,
 		TEXT("Failed to create sandboxed mount directory %s for mount point %s"),
-		*AbsoluteSandboxPath, *AbsoluteNonSandboxPath
+		*MountPointPath.GetSandboxPath(), *MountPointPath.GetNonSandboxPath()
 		);
 	
 	if (bCreatedDirectory)
 	{
 		const FScopeLock SandboxMountPointsLock(&SandboxMountPointsCS);
 
-		FSandboxMountPoint& SandboxMountPoint = SandboxMountPoints.Add_GetRef(
-			FSandboxMountPoint{ FSandboxedPlatformFilePath(MoveTemp(AbsoluteNonSandboxPath), MoveTemp(AbsoluteSandboxPath)) }
-			);
+		FSandboxMountPoint& SandboxMountPoint = SandboxMountPoints.Add_GetRef(FSandboxMountPoint{ MountPointPath });
 #if UE_SANDBOX_WITH_DIRECTORY_WATCHER
 		if (IDirectoryWatcher* DirectoryWatcher = GetDirectoryWatcher())
 		{
diff --git a/Engine/Plugins/Developer/Sandbox/FileSandbox/Source/FileSandboxCore/Private/Sandbox/Platform/SandboxPlatformFile.h b/Engine/Plugins/Developer/Sandbox/FileSandbox/Source/FileSandboxCore/Private/Sandbox/Platform/SandboxPlatformFile.h
--- a/Engine/Plugins/Developer/Sandbox/FileSandbox/Source/FileSandboxCore/Private/Sandbox/Platform/SandboxPlatformFile.h
+++ b/Engine/Plugins/Developer/Sandbox/FileSandbox/Source/FileSandboxCore/Private/Sandbox/Platform/SandboxPlatformFile.h
@@ -97,12 +97,6 @@ public:
 		TArray<FName>& OutPackagesPendingPurge,
 		TFunctionRef<void(const FFileChange&)> InProcessFileChangeCallback = [](const FFileChange&){}
 		);
-
-	/** 
-	 * Gather all files changes that are currently in the sandbox.
-	 * @return Array of filepaths in no particular order.
-	 */
-	TArray<FString> GatherSandboxChangedFilenames() const;
 	
 	/** @return Whether the given package file was deleted and exists on non sandbox path. */
 	bool DeletedPackageExistsInNonSandbox(FString InFilename) const;
diff --git a/Engine/Plugins/Developer/Sandbox/FileSandbox/Source/FileSandboxCore/Private/Sandbox/Platform/SandboxedPlatformFilePath.h b/Engine/Plugins/Developer/Sandbox/FileSandbox/Source/FileSandboxCore/Private/Sandbox/Platform/SandboxedPlatformFilePath.h
--- a/Engine/Plugins/Developer/Sandbox/FileSandbox/Source/FileSandboxCore/Private/Sandbox/Platform/SandboxedPlatformFilePath.h
+++ b/Engine/Plugins/Developer/Sandbox/FileSandbox/Source/FileSandboxCore/Private/Sandbox/Platform/SandboxedPlatformFilePath.h
@@ -3,6 +3,7 @@
 #pragma once
 
 #include "Containers/UnrealString.h"
+#include "Misc/Paths.h"
 
 namespace UE::FileSandboxCore
 {
@@ -26,6 +27,21 @@ public:
 	{
 	}
 
+	/**
+	 * Creates a platform file path that maps from an engine mount point to sandboxed mount point
+	 * @param InSandboxMountPointPath Directory in the sandbox file structure where mount points are stored, i.e. "MySandbox/Sandbox"
+	 * @param InAssetPath Name of the mount point, such as "Game", "Engine", "MyPlugin", etc.
+	 * @param InFilesystemPath Root directory such as "MyProject/Content", "Engine/Content", "MyPlugin/Content", etc.
+	 */
+	static FSandboxedPlatformFilePath CreateMountPoint(
+		const FString& InSandboxMountPointPath, const FString& InAssetPath, const FString& InFilesystemPath
+		)
+	{
+		FString AbsoluteSandboxPath = FPaths::ConvertRelativePathToFull(InSandboxMountPointPath / InAssetPath) / TEXT("");
+		FString AbsoluteNonSandboxPath = FPaths::ConvertRelativePathToFull(InFilesystemPath) / TEXT("");
+		return FSandboxedPlatformFilePath(MoveTemp(AbsoluteNonSandboxPath), MoveTemp(AbsoluteSandboxPath));
+	}
+
 	static FSandboxedPlatformFilePath CreateSandboxPath(FString&& InNonSandboxPath, const FSandboxedPlatformFilePath& InRootPath)
 	{
 		checkf(InRootPath.HasSandboxPath(), TEXT("Root '%s' had no sandbox path set!"), *InRootPath.GetNonSandboxPath());
diff --git a/Engine/Plugins/Developer/Sandbox/FileSandbox/Source/FileSandboxCore/Private/Sandbox/SandboxInstance.cpp b/Engine/Plugins/Developer/Sandbox/FileSandbox/Source/FileSandboxCore/Private/Sandbox/SandboxInstance.cpp
--- a/Engine/Plugins/Developer/Sandbox/FileSandbox/Source/FileSandboxCore/Private/Sandbox/SandboxInstance.cpp
+++ b/Engine/Plugins/Developer/Sandbox/FileSandbox/Source/FileSandboxCore/Private/Sandbox/SandboxInstance.cpp
@@ -6,14 +6,16 @@
 #include "Misc/PackageName.h"
 #include "Persist/AggregatePersistErrorHandler.h"
 #include "Persist/AppendTextErrorHandler.h"
+#include "Persist/DiffUtils.h"
 #include "Persist/LogPersistErrorHandler.h"
 #include "Types/Manager/SandboxInitArgs.h"
 #include "Types/Sandbox/DiscardResult.h"
 #include "Types/SandboxFileChange.h"
-#include "Utils/AssetUtils.h"
-#include "Utils/FileChange/FileChange.h"
 #include "Utils/PackageSandboxUtils.h"
+#include "Utils/SandboxFileUtils.h"
 #include "Utils/SandboxDirectoryUtils.h"
+#include "Utils/AssetUtils.h"
+#include "Utils/FileChange/FileChange.h"
 #include "Utils/SandboxFileUtils.h"
 
 namespace UE::FileSandboxCore
@@ -146,9 +148,9 @@ FDiscardResult FSandboxInstance::DiscardSandbox(const FDiscardArgs& InArgs)
 	return Result;
 }
 
-TArray<FString> FSandboxInstance::GatherChangedFiles() const
+void FSandboxInstance::EnumerateFileChanges(TFunctionRef<FProcessFileChangeSignature> InProcess) const
 {
-	return PlatformFile->GatherSandboxChangedFilenames();
+	FileSandboxCore::EnumerateFileChanges(*PlatformFile->GetLowerLevel(), GetRootDirectory(), ManifestContent, InProcess);
 }
 
 bool FSandboxInstance::DeletedPackageExistsInNonSandbox(const FString& InFilename) const
diff --git a/Engine/Plugins/Developer/Sandbox/FileSandbox/Source/FileSandboxCore/Private/Sandbox/SandboxInstance.h b/Engine/Plugins/Developer/Sandbox/FileSandbox/Source/FileSandboxCore/Private/Sandbox/SandboxInstance.h
--- a/Engine/Plugins/Developer/Sandbox/FileSandbox/Source/FileSandboxCore/Private/Sandbox/SandboxInstance.h
+++ b/Engine/Plugins/Developer/Sandbox/FileSandbox/Source/FileSandboxCore/Private/Sandbox/SandboxInstance.h
@@ -37,7 +37,7 @@ public:
 	//~ Begin ISandboxInstance Interface
 	virtual FPersistResult PersistSandbox(const FPersistArgs& InArgs) override;
 	virtual FDiscardResult DiscardSandbox(const FDiscardArgs& InArgs) override;
-	virtual TArray<FString> GatherChangedFiles() const override;
+	virtual void EnumerateFileChanges(TFunctionRef<FProcessFileChangeSignature> InProcess) const override;
 	virtual bool DeletedPackageExistsInNonSandbox(const FString& InFilename) const override;
 	virtual const FFileSandboxCore_SandboxMetaData& GetInitialMetaData() const override { return MetaData; }
 	virtual const TCHAR* GetRootDirectory() const override { return *RootDirectory; }
diff --git a/Engine/Plugins/Developer/Sandbox/FileSandbox/Source/FileSandboxCore/Private/Sandbox/SandboxManager.cpp b/Engine/Plugins/Developer/Sandbox/FileSandbox/Source/FileSandboxCore/Private/Sandbox/SandboxManager.cpp
--- a/Engine/Plugins/Developer/Sandbox/FileSandbox/Source/FileSandboxCore/Private/Sandbox/SandboxManager.cpp
+++ b/Engine/Plugins/Developer/Sandbox/FileSandbox/Source/FileSandboxCore/Private/Sandbox/SandboxManager.cpp
@@ -7,6 +7,9 @@
 #include "LogFileSandbox.h"
 #include "Misc/CoreDelegates.h"
 #include "Misc/Paths.h"
+#include "Persist/DiffUtils.h"
+#include "Types/EBreakBehavior.h"
+#include "Types/SandboxFileChange.h"
 #include "Types/Manager/DeleteSandboxByDirectoryArgs.h"
 #include "Types/Manager/DeleteSandboxResult.h"
 #include "Types/Manager/LeaveSandboxResult.h"
@@ -129,6 +132,22 @@ FLeaveSandboxResult FSandboxManager::LeaveSandbox(const FLeaveSandboxArgs& InLea
 	return ELeaveSandboxErrorCode::Success;
 }
 
+void FSandboxManager::EnumerateFileChanges(const FString& InSandboxRootPath, TFunctionRef<FProcessFileChangeSignature> InProcess) const
+{
+	ISandboxInstance* Instance = GetActiveSandboxInstance();
+	if (Instance && FPaths::IsSamePath(Instance->GetRootDirectory(), InSandboxRootPath))
+	{
+		Instance->EnumerateFileChanges(InProcess);
+	}
+	else if (const TOptional<FFileSandboxCore_ManifestData> Manifest = LoadManifest(InSandboxRootPath))
+	{
+		IPlatformFile& PlatformFile = Instance 
+			? *FPlatformFileManager::Get().GetPlatformFile().GetLowerLevel() 
+			: FPlatformFileManager::Get().GetPlatformFile();
+		FileSandboxCore::EnumerateFileChanges(PlatformFile, InSandboxRootPath, *Manifest, InProcess);
+	}
+}
+
 void FSandboxManager::OnEngineExit()
 {
 	LeaveSandbox();
diff --git a/Engine/Plugins/Developer/Sandbox/FileSandbox/Source/FileSandboxCore/Private/Sandbox/SandboxManager.h b/Engine/Plugins/Developer/Sandbox/FileSandbox/Source/FileSandboxCore/Private/Sandbox/SandboxManager.h
--- a/Engine/Plugins/Developer/Sandbox/FileSandbox/Source/FileSandboxCore/Private/Sandbox/SandboxManager.h
+++ b/Engine/Plugins/Developer/Sandbox/FileSandbox/Source/FileSandboxCore/Private/Sandbox/SandboxManager.h
@@ -23,6 +23,7 @@ public:
 	virtual FLeaveSandboxResult LeaveSandbox(const FLeaveSandboxArgs& InLeaveArgs = FLeaveSandboxArgs()) override;
 	virtual ISandboxLock* GetActiveLock() const override { return ActiveSandboxData ? ActiveSandboxData->Lock.Get() : nullptr; }
 	virtual ISandboxInstance* GetActiveSandboxInstance() const override { return ActiveSandboxData ? ActiveSandboxData->Sandbox.Get() : nullptr; }
+	virtual void EnumerateFileChanges(const FString& InSandboxRootPath, TFunctionRef<FProcessFileChangeSignature> InProcess) const override;
 	virtual FSandboxInstanceEvent& OnPostSandboxStartup() override { return OnPostSandboxStartupDelegate; }
 	virtual FSandboxInstanceEvent& OnPreSandboxShutdown() override { return OnPreSandboxShutdownDelegate; }
 	virtual FSandboxShutdownEvent& OnPostSandboxShutdown() override { return OnPostSandboxShutdownDelegate; }
diff --git a/Engine/Plugins/Developer/Sandbox/FileSandbox/Source/FileSandboxCore/Private/Utils/SandboxFileUtils.cpp b/Engine/Plugins/Developer/Sandbox/FileSandbox/Source/FileSandboxCore/Private/Utils/SandboxFileUtils.cpp
--- a/Engine/Plugins/Developer/Sandbox/FileSandbox/Source/FileSandboxCore/Private/Utils/SandboxFileUtils.cpp
+++ b/Engine/Plugins/Developer/Sandbox/FileSandbox/Source/FileSandboxCore/Private/Utils/SandboxFileUtils.cpp
@@ -47,4 +47,9 @@ FString GetMetadataFileName()
 {
 	return TEXT("metadata.json");
 }
+
+FString GetSandboxMountPointRoot(const FString& InSandboxRootPath)
+{
+	return InSandboxRootPath / TEXT("Sandbox");
+}
 }
diff --git a/Engine/Plugins/Developer/Sandbox/FileSandbox/Source/FileSandboxCore/Private/Utils/SandboxFileUtils.h b/Engine/Plugins/Developer/Sandbox/FileSandbox/Source/FileSandboxCore/Private/Utils/SandboxFileUtils.h
--- a/Engine/Plugins/Developer/Sandbox/FileSandbox/Source/FileSandboxCore/Private/Utils/SandboxFileUtils.h
+++ b/Engine/Plugins/Developer/Sandbox/FileSandbox/Source/FileSandboxCore/Private/Utils/SandboxFileUtils.h
@@ -17,4 +17,7 @@ TOptional<FFileSandboxCore_ManifestData> LoadManifest(const FString& InRootDirec
 FString GetManifestFileName();
 /* @return The name to give the metadata file. */
 FString GetMetadataFileName();
+
+/** @return Path to directory in the sandbox file structure where mount points are stored, i.e. "MySandbox/Sandbox" */
+FString GetSandboxMountPointRoot(const FString& InSandboxRootPath);
 }
\ No newline at end of file
diff --git a/Engine/Plugins/Developer/Sandbox/FileSandbox/Source/FileSandboxCore/Public/ISandboxInstance.h b/Engine/Plugins/Developer/Sandbox/FileSandbox/Source/FileSandboxCore/Public/ISandboxInstance.h
--- a/Engine/Plugins/Developer/Sandbox/FileSandbox/Source/FileSandboxCore/Public/ISandboxInstance.h
+++ b/Engine/Plugins/Developer/Sandbox/FileSandbox/Source/FileSandboxCore/Public/ISandboxInstance.h
@@ -2,8 +2,10 @@
 
 #pragma once
 
+#include "Types/GatheredFileChanges.h"
 #include "Types/Sandbox/DiscardArgs.h"
 #include "Types/Sandbox/PersistResult.h"
+#include "Types/ArgsFwds.h"
 
 #define UE_API FILESANDBOXCORE_API
 
@@ -11,6 +13,7 @@ struct FFileSandboxCore_SandboxMetaData;
 
 namespace UE::FileSandboxCore
 {
+enum class EFileChangeGatherFlags : uint8;
 enum class EBreakBehavior : uint8;
 struct FDiscardResult;
 struct FPersistArgs;
@@ -55,7 +58,13 @@ public:
 	
 	
 	/** Enumerates all files that have been changed (added, removed, or edited). */
-	virtual TArray<FString> GatherChangedFiles() const = 0;
+	UE_API FGatheredFileChanges GatherChangedFiles(EFileChangeGatherFlags InFlags = EFileChangeGatherFlags::All) const;
+	
+	/** 
+	 * Invokes InProcess for every effective file change that would be performed if the sandbox was persisted.
+	 * @see ISandboxManager::EnumerateFileChanges.
+	 */
+	virtual void EnumerateFileChanges(TFunctionRef<FProcessFileChangeSignature> InProcess) const = 0;
 
 	
 	/** @return Whether the given package file exists on non sandbox path, i.e. whether the file was deleted in sandbox mode. */
diff --git a/Engine/Plugins/Developer/Sandbox/FileSandbox/Source/FileSandboxCore/Public/ISandboxManager.h b/Engine/Plugins/Developer/Sandbox/FileSandbox/Source/FileSandboxCore/Public/ISandboxManager.h
--- a/Engine/Plugins/Developer/Sandbox/FileSandbox/Source/FileSandboxCore/Public/ISandboxManager.h
+++ b/Engine/Plugins/Developer/Sandbox/FileSandbox/Source/FileSandboxCore/Public/ISandboxManager.h
@@ -7,13 +7,16 @@
 #include "Types/Manager/LeaveSandboxArgs.h"
 #include "Types/Manager/LeaveSandboxResult.h"
 #include "Types/Manager/SandboxCreationResult.h"
+#include "Utils/SandboxDirectoryUtils.h"
 
 #define UE_API FILESANDBOXCORE_API
 
 namespace UE::FileSandboxCore
 {
-class ISandboxLock;
 class ISandboxInstance;
+class ISandboxLock;
+enum class EBreakBehavior : uint8;
+enum class ESandboxFileChange : uint8;
 	
 DECLARE_MULTICAST_DELEGATE_OneParam(FSandboxInstanceEvent, ISandboxInstance&);
 using FSandboxShutdownEvent = FSimpleMulticastDelegate;
@@ -118,7 +121,27 @@ public:
 
 	/** @return Whether the is any active sandbox. */
 	bool HasActiveSandbox() const { return GetActiveSandboxInstance() != nullptr; }
-
+	
+	/** 
+	 *  Invokes InProcess for every effective file change that would be performed if the sandbox was persisted.
+	 * 
+	 * Effective changes can differ from the actual actions performed during sandbox because external changes could have been made to the project
+	 * outside of sandbox.
+	 * 
+	 * The implementation respects file changes that have occured since the sandbox the last time the sandbox was active.
+	 * Examples: Suppose that while the engine is sandboxed file x is
+	 * - removed: then, outside of sandbox, file x is added. Then X would not be listed (InProcess not invoked).
+	 * - added: then, outside of sandbox, file x is added. Then X would be reported as edited.
+	 * - edited: then, outside of sandbox, file x is removed. Then X would be reported as added.
+	 * If no external changes are made, the file action is reported exactly as it was performed in the sandbox without changes.
+	 */
+	UE_API void EnumerateFileChangesByName(
+		const FString& InSandboxName, TFunctionRef<FProcessFileChangeSignature> InProcess,
+		const FString& InBaseDirectory = GetBaseSandboxDirectory()
+		) const;
+	
+	/** Alternative version that uses the sandbox root directory instead.*/
+	virtual void EnumerateFileChanges(const FString& InSandboxRootPath, TFunctionRef<FProcessFileChangeSignature> InProcess) const = 0;
 	
 	/** @return Event invoked after a sandbox has been successfully created. */
 	virtual FSandboxInstanceEvent& OnPostSandboxStartup() = 0;
diff --git a/Engine/Plugins/Developer/Sandbox/FileSandbox/Source/FileSandboxCore/Public/Types/ArgsFwds.h b/Engine/Plugins/Developer/Sandbox/FileSandbox/Source/FileSandboxCore/Public/Types/ArgsFwds.h
--- a/Engine/Plugins/Developer/Sandbox/FileSandbox/Source/FileSandboxCore/Public/Types/ArgsFwds.h
+++ b/Engine/Plugins/Developer/Sandbox/FileSandbox/Source/FileSandboxCore/Public/Types/ArgsFwds.h
@@ -2,6 +2,10 @@
 
 #pragma once
 
+#include "HAL/Platform.h"
+
+class FString;
+
 namespace UE::FileSandboxCore
 {
 struct FDeleteSandboxByDirectoryArgs;
@@ -20,4 +24,9 @@ template<typename TErrorType>
 struct TSandboxCreationResult;
 using FNewSandboxResult = TSandboxCreationResult<FNewSandboxError>;
 using FLoadSandboxResult = TSandboxCreationResult<FLoadSandboxError>;
+
+enum class EBreakBehavior : uint8;
+enum class ESandboxFileChange : uint8;
+
+using FProcessFileChangeSignature = EBreakBehavior(const FString& InPath, ESandboxFileChange InChange);
 }
\ No newline at end of file
diff --git a/Engine/Plugins/Developer/Sandbox/FileSandbox/Source/FileSandboxCore/Public/Types/GatheredFileChanges.h b/Engine/Plugins/Developer/Sandbox/FileSandbox/Source/FileSandboxCore/Public/Types/GatheredFileChanges.h
new file mode 100644
--- /dev/null
+++ b/Engine/Plugins/Developer/Sandbox/FileSandbox/Source/FileSandboxCore/Public/Types/GatheredFileChanges.h
@@ -0,0 +1,35 @@
+﻿// Copyright Epic Games, Inc. All Rights Reserved.
+
+#pragma once
+
+#include "Containers/Array.h"
+#include "HAL/Platform.h"
+
+class FString;
+
+namespace UE::FileSandboxCore
+{
+enum class ESandboxFileChange : uint8;
+
+/** Specifies what information should be gathered about files */
+enum class EFileChangeGatherFlags : uint8
+{
+	None,
+	IncludeChangeTypes = 1 << 0,
+	All = IncludeChangeTypes
+};
+ENUM_CLASS_FLAGS(EFileChangeGatherFlags);
+
+/** Result of ISandboxInstance::GatherChangedFiles. */
+struct FGatheredFileChanges
+{
+	/** All files changed by the sandbox. The paths are in the non-sandbox. */
+	TArray<FString> NonSandboxPaths;
+	
+	/** The changes made to the file. Equal length to NonSandboxPaths. Only set if EFileChangeGatherFlags::IncludeChangeTypes was set. */
+	TArray<ESandboxFileChange> FileActions;
+	
+	/** @return Whether there are any changes. */
+	bool HasChanges() const { return !NonSandboxPaths.IsEmpty(); }
+};
+}
\ No newline at end of file
