File Sandbox Core: Add command line arguments.

Also added support for Sandboxed Editing to treat sandboxes started via command line to be treated as if they are owned by [REDACTED_PERSON].

Arguments:
-DELETESANDBOXBYNAME
-DELETESANDBOXBYPATH
-STARTWITHSANDBOX
-SANDBOXDESCRIPTION
-DEFAULTSANDBOXDIRECTORY

#Synced-CL 49159615

diff --git a/Engine/Plugins/Developer/Sandbox/FileSandbox/Source/FileSandboxCore/Private/FileSandboxCoreModule.cpp b/Engine/Plugins/Developer/Sandbox/FileSandbox/Source/FileSandboxCore/Private/FileSandboxCoreModule.cpp
--- a/Engine/Plugins/Developer/Sandbox/FileSandbox/Source/FileSandboxCore/Private/FileSandboxCoreModule.cpp
+++ b/Engine/Plugins/Developer/Sandbox/FileSandbox/Source/FileSandboxCore/Private/FileSandboxCoreModule.cpp
@@ -2,7 +2,9 @@
 
 #include "FileSandboxCoreModule.h"
 
+#include "Misc/CoreDelegates.h"
 #include "Repository/NaiveSandboxRepository.h"
+#include "Utils/CommandLineUtils.h"
 #include "Utils/SandboxDirectoryUtils.h"
 
 #if UE_SANDBOX_WITH_DIRECTORY_WATCHER
@@ -20,6 +22,8 @@ void FFileSandboxCoreModule::StartupModule()
 #else
 	DefaultRepository = MakeUnique<FNaiveSandboxRepository>(GetBaseSandboxDirectory(), *SandboxManager);
 #endif
+	
+	RegisterStartupCommandLineDelegate();
 }
 
 void FFileSandboxCoreModule::ShutdownModule()
diff --git a/Engine/Plugins/Developer/Sandbox/FileSandbox/Source/FileSandboxCore/Private/Utils/CommandLineUtils.cpp b/Engine/Plugins/Developer/Sandbox/FileSandbox/Source/FileSandboxCore/Private/Utils/CommandLineUtils.cpp
new file mode 100644
--- /dev/null
+++ b/Engine/Plugins/Developer/Sandbox/FileSandbox/Source/FileSandboxCore/Private/Utils/CommandLineUtils.cpp
@@ -0,0 +1,119 @@
+﻿// Copyright Epic Games, Inc. All Rights Reserved.
+
+#include "CommandLineUtils.h"
+
+#include "CoreGlobals.h"
+#include "IFileSandboxCoreModule.h"
+#include "ISandboxManager.h"
+#include "LogFileSandbox.h"
+#include "Misc/CommandLine.h"
+#include "Misc/CoreDelegates.h"
+#include "Misc/Paths.h"
+#include "Types/Manager/DeleteSandboxByDirectoryArgs.h"
+#include "Types/Manager/DeleteSandboxByNameArgs.h"
+#include "Types/Manager/DeleteSandboxResult.h"
+#include "Types/Manager/NewSandboxArgs.h"
+#include "Utils/BuiltInTags.h"
+
+namespace UE::FileSandboxCore
+{
+namespace CommandDetail
+{
+/** Deletes the sandbox with the given name. */
+const TCHAR* DeleteSandboxByName = TEXT("-DELETESANDBOXBYNAME=");
+/** Deletes the sandbox with the given path. */
+const TCHAR* DeleteSandboxByPath = TEXT("-DELETESANDBOXBYPATH=");
+
+/** Starts a new sandbox with the given name. */
+const TCHAR* StartWithSandbox = TEXT("-STARTWITHSANDBOX=");
+/** Optional description to give the new sandbox. */
+const TCHAR* SandboxDescription = TEXT("-SANDBOXDESCRIPTION=");
+/** Optional base path to place the sandbox in. Uses the default directory if unspecified. */
+const TCHAR* SandboxBaseDirectory = TEXT("-SANDBOXBASEDIRECTORY=");
+
+/** Sets the folder in which new sandboxes are placed by default. */
+const TCHAR* DefaultSandboxDirectory = TEXT("-DEFAULTSANDBOXDIRECTORY=");
+
+static void HandleDeleteNamedSandbox()
+{
+	FString Name;
+	if (FParse::Value(FCommandLine::Get(), DeleteSandboxByName, Name))
+	{
+		ISandboxManager& SandboxManager = IFileSandboxCoreModule::Get().GetSandboxManager();
+		const FDeleteSandboxResult DeleteResult = SandboxManager.DeleteNamedSandbox(FDeleteSandboxByNameArgs(Name));
+		
+		UE_CLOG(DeleteResult.IsSuccess(), LogFileSandbox, Log, TEXT("Deleted named sandbox '%s' from command line"), *Name);
+		UE_CLOG(!DeleteResult.IsSuccess(), LogFileSandbox, Error, TEXT("Failed to delete named sandbox '%s from command line'"), *Name);
+	}
+}
+
+static void HandleDeleteSandboxByPath()
+{
+	FString Path;
+	if (FParse::Value(FCommandLine::Get(), DeleteSandboxByPath, Path))
+	{
+		ISandboxManager& SandboxManager = IFileSandboxCoreModule::Get().GetSandboxManager();
+		const FDeleteSandboxResult DeleteResult = SandboxManager.DeleteSandbox(
+			FDeleteSandboxByDirectoryArgs(Path)
+			);
+		
+		UE_CLOG(DeleteResult.IsSuccess(), LogFileSandbox, Log, TEXT("Deleted sandbox by path '%s' from command line"), *Path);
+		UE_CLOG(!DeleteResult.IsSuccess(), LogFileSandbox, Error, TEXT("Failed sandbox by path '%s from command line'"), *Path);
+	}
+}
+
+static void HandleStartWithSandbox()
+{
+	FString Name;
+	if (!FParse::Value(FCommandLine::Get(), StartWithSandbox, Name))
+	{
+		return;
+	}
+	
+	if (Name.IsEmpty())
+	{
+		UE_LOG(LogFileSandbox, Error, TEXT("Cannot start sandbox  from command line because no name was specified"));
+		return;
+	}
+	
+	FString Description = TEXT("Command Line Sandbox");
+	FParse::Value(FCommandLine::Get(), SandboxDescription, Description);
+	
+	FString BaseDirectory = GetBaseSandboxDirectory();
+	FParse::Value(FCommandLine::Get(), SandboxBaseDirectory, BaseDirectory);
+		
+	FNewSandboxArgs Args(Name, Description);
+	Args.SandboxBasePath = BaseDirectory;
+	Args.MetaData.Tags.Add(Tag_CommandLineSandbox);
+	ISandboxManager& SandboxManager = IFileSandboxCoreModule::Get().GetSandboxManager();
+	const FNewSandboxResult Result = SandboxManager.CreateNewSandbox(Args);
+		
+	UE_CLOG(Result.HasValue(), LogFileSandbox, Log, TEXT("Started sandbox '%s' from command line"), *Name);
+	UE_CLOG(Result.HasError(), LogFileSandbox, Error, TEXT("Failed to start sandbox '%s' from command line'"), *Name);
+}
+}
+
+void RegisterStartupCommandLineDelegate()
+{
+	if (!IsRunningCommandlet())
+	{
+		FCoreDelegates::OnFEngineLoopInitComplete.AddStatic(&ParseStartupCommandLine);
+	}
+}
+
+void ParseStartupCommandLine()
+{
+	CommandDetail::HandleDeleteNamedSandbox();
+	CommandDetail::HandleDeleteSandboxByPath();
+	CommandDetail::HandleStartWithSandbox();
+}
+
+FString ParseDefaultSandboxDirectory()
+{
+	FString CommandLineDirectory;
+	const bool bHasOverride = FParse::Value(FCommandLine::Get(), CommandDetail::DefaultSandboxDirectory, CommandLineDirectory);
+	return bHasOverride && FPaths::ValidatePath(CommandLineDirectory) 
+		? CommandLineDirectory 
+		: FString();
+}
+}
\ No newline at end of file
diff --git a/Engine/Plugins/Developer/Sandbox/FileSandbox/Source/FileSandboxCore/Private/Utils/CommandLineUtils.h b/Engine/Plugins/Developer/Sandbox/FileSandbox/Source/FileSandboxCore/Private/Utils/CommandLineUtils.h
new file mode 100644
--- /dev/null
+++ b/Engine/Plugins/Developer/Sandbox/FileSandbox/Source/FileSandboxCore/Private/Utils/CommandLineUtils.h
@@ -0,0 +1,17 @@
+﻿// Copyright Epic Games, Inc. All Rights Reserved.
+
+#pragma once
+
+class FString;
+
+namespace UE::FileSandboxCore
+{
+/** Parses the command line once engine start up is done; this is only done if command line arguments should be parsed, i.e. not running commandlet. */
+void RegisterStartupCommandLineDelegate();
+
+/** Parses command line commands that are intended for engine startup. */
+void ParseStartupCommandLine();
+
+/** @return The default sandbox directory set via command line, or empty no (valid) directory is set. */
+FString ParseDefaultSandboxDirectory();
+}
diff --git a/Engine/Plugins/Developer/Sandbox/FileSandbox/Source/FileSandboxCore/Private/Utils/SandboxDirectoryUtils.cpp b/Engine/Plugins/Developer/Sandbox/FileSandbox/Source/FileSandboxCore/Private/Utils/SandboxDirectoryUtils.cpp
--- a/Engine/Plugins/Developer/Sandbox/FileSandbox/Source/FileSandboxCore/Private/Utils/SandboxDirectoryUtils.cpp
+++ b/Engine/Plugins/Developer/Sandbox/FileSandbox/Source/FileSandboxCore/Private/Utils/SandboxDirectoryUtils.cpp
@@ -2,12 +2,14 @@
 
 #include "Utils/SandboxDirectoryUtils.h"
 
+#include "CommandLineUtils.h"
 #include "Data/ManifestData.h"
 #include "Data/SandboxMetaData.h"
 #include "Data/VersionInfo.h"
 #include "HAL/PlatformFileManager.h"
 #include "JsonObjectConverter.h"
 #include "LogFileSandbox.h"
+#include "Misc/CommandLine.h"
 #include "Misc/FileHelper.h"
 #include "Misc/Paths.h"
 #include "SandboxFileUtils.h"
@@ -20,7 +22,10 @@ namespace UE::FileSandboxCore
 {
 FString GetBaseSandboxDirectory()
 {
-	return FPaths::ProjectIntermediateDir() / TEXT("Sandboxes");
+	const FString CommandLineOverride = ParseDefaultSandboxDirectory();
+	return CommandLineOverride.IsEmpty() 
+		? FPaths::ProjectIntermediateDir() / TEXT("Sandboxes")
+		: CommandLineOverride;
 }
 	
 void ForEachSandbox(
diff --git a/Engine/Plugins/Developer/Sandbox/FileSandbox/Source/FileSandboxCore/Public/Utils/BuiltInTags.h b/Engine/Plugins/Developer/Sandbox/FileSandbox/Source/FileSandboxCore/Public/Utils/BuiltInTags.h
new file mode 100644
--- /dev/null
+++ b/Engine/Plugins/Developer/Sandbox/FileSandbox/Source/FileSandboxCore/Public/Utils/BuiltInTags.h
@@ -0,0 +1,11 @@
+﻿// Copyright Epic Games, Inc. All Rights Reserved.
+
+#pragma once
+
+#include "UObject/NameTypes.h"
+
+namespace UE::FileSandboxCore
+{
+/** Tag added to sandboxes that were created from command line. */
+constexpr FLazyName Tag_CommandLineSandbox("CommandLineSandbox");
+}
diff --git a/Engine/Plugins/Developer/Sandbox/SandboxedEditing/Source/SandboxedEditing/Private/Framework/EntryPoint/SandboxedEditingEntryPoint.cpp b/Engine/Plugins/Developer/Sandbox/SandboxedEditing/Source/SandboxedEditing/Private/Framework/EntryPoint/SandboxedEditingEntryPoint.cpp
--- a/Engine/Plugins/Developer/Sandbox/SandboxedEditing/Source/SandboxedEditing/Private/Framework/EntryPoint/SandboxedEditingEntryPoint.cpp
+++ b/Engine/Plugins/Developer/Sandbox/SandboxedEditing/Source/SandboxedEditing/Private/Framework/EntryPoint/SandboxedEditingEntryPoint.cpp
@@ -23,7 +23,7 @@ FText FSandboxedEditingEntryPoint::GetEntryPointLabel() const
 
 bool FSandboxedEditingEntryPoint::OwnsSandbox(const FileSandboxCore::ISandboxInstance& InSandbox) const
 {
-	return IsOwnedBySandboxedEditing(InSandbox.GetInitialMetaData());
+	return CanBeShownBySandboxedEditing(InSandbox.GetInitialMetaData());
 }
 }
 
diff --git a/Engine/Plugins/Developer/Sandbox/SandboxedEditing/Source/SandboxedEditing/Private/Framework/Models/SandboxSystemModel.cpp b/Engine/Plugins/Developer/Sandbox/SandboxedEditing/Source/SandboxedEditing/Private/Framework/Models/SandboxSystemModel.cpp
--- a/Engine/Plugins/Developer/Sandbox/SandboxedEditing/Source/SandboxedEditing/Private/Framework/Models/SandboxSystemModel.cpp
+++ b/Engine/Plugins/Developer/Sandbox/SandboxedEditing/Source/SandboxedEditing/Private/Framework/Models/SandboxSystemModel.cpp
@@ -239,7 +239,7 @@ FString FSandboxSystemModel::GetActiveSandboxName() const
 	using namespace FileSandboxCore;
 	IFileSandboxCoreModule& Module = IFileSandboxCoreModule::Get();
 	ISandboxInstance* Instance = Module.GetSandboxManager().GetActiveSandboxInstance();
-	return Instance && IsOwnedBySandboxedEditing(Instance->GetInitialMetaData()) ? Instance->GetInitialMetaData().Name : FString();
+	return Instance && CanBeShownBySandboxedEditing(Instance->GetInitialMetaData()) ? Instance->GetInitialMetaData().Name : FString();
 }
 
 bool FSandboxSystemModel::HasActiveSandbox() const
@@ -256,7 +256,7 @@ void FSandboxSystemModel::ForEachSandbox(TFunctionRef<void(const FSandboxInfo& I
 	Repository.ForEachSandbox([&InProcess](const FString& InRootPath, const FSandboxMetaInfo& MetaInfo)
 	{
 		const FFileSandboxCore_SandboxMetaData& MetaData = MetaInfo.UserMetaData;
-		if (IsOwnedBySandboxedEditing(MetaData))
+		if (CanBeShownBySandboxedEditing(MetaData))
 		{
 			InProcess(FSandboxInfo(MetaData.Name, MetaData.Description, InRootPath, GetLastModified(InRootPath), MetaInfo.VersionInfo.EngineVersion));
 		}
@@ -281,7 +281,7 @@ TOptional<FSandboxInfo> FSandboxSystemModel::GetSandboxInfo(const FString& InSan
 	Repository.ReadMetaData(InSandboxRoot, [&InSandboxRoot, &Result](const FSandboxMetaInfo& MetaInfo)
 	{
 		const FFileSandboxCore_SandboxMetaData& MetaData = MetaInfo.UserMetaData;
-		if (IsOwnedBySandboxedEditing(MetaData))
+		if (CanBeShownBySandboxedEditing(MetaData))
 		{
 			Result.Emplace(MetaData.Name, MetaData.Description, InSandboxRoot, MetaInfo.LastModified, MetaInfo.VersionInfo.EngineVersion);
 		}
@@ -296,7 +296,7 @@ void FSandboxSystemModel::SetDescription(const FString& InSandboxRoot, const FSt
 	ISandboxRepository& Repository = Module.GetDefaultSandboxRepository();
 	Repository.ReadMetaData(InSandboxRoot, [&InSandboxRoot, &InDescription](const FSandboxMetaInfo& MetaInfo)
 	{
-		if (IsOwnedBySandboxedEditing(MetaInfo.UserMetaData))
+		if (CanBeShownBySandboxedEditing(MetaInfo.UserMetaData))
 		{
 			FFileSandboxCore_SandboxMetaData MetaData = MetaInfo.UserMetaData;
 			MetaData.Description = InDescription;
diff --git a/Engine/Plugins/Developer/Sandbox/SandboxedEditing/Source/SandboxedEditing/Private/Framework/Models/SandboxTagging.cpp b/Engine/Plugins/Developer/Sandbox/SandboxedEditing/Source/SandboxedEditing/Private/Framework/Models/SandboxTagging.cpp
--- a/Engine/Plugins/Developer/Sandbox/SandboxedEditing/Source/SandboxedEditing/Private/Framework/Models/SandboxTagging.cpp
+++ b/Engine/Plugins/Developer/Sandbox/SandboxedEditing/Source/SandboxedEditing/Private/Framework/Models/SandboxTagging.cpp
@@ -3,6 +3,7 @@
 #include "SandboxTagging.h"
 
 #include "Data/SandboxMetaData.h"
+#include "Utils/BuiltInTags.h"
 
 namespace UE::SandboxedEditing
 {
@@ -11,6 +12,12 @@ bool IsOwnedBySandboxedEditing(const FFileSandboxCore_SandboxMetaData& InMetaDat
 	return InMetaData.Tags.Contains(SandboxedEditingTag.Resolve());
 }
 
+bool CanBeShownBySandboxedEditing(const FFileSandboxCore_SandboxMetaData& InMetaData)
+{
+	return IsOwnedBySandboxedEditing(InMetaData) 
+		|| InMetaData.Tags.Contains(FileSandboxCore::Tag_CommandLineSandbox.Resolve());
+}
+
 void MarkOwnedBySandboxedEditing(FFileSandboxCore_SandboxMetaData& InMetaData)
 {
 	InMetaData.Tags.Add(SandboxedEditingTag.Resolve());
diff --git a/Engine/Plugins/Developer/Sandbox/SandboxedEditing/Source/SandboxedEditing/Private/Framework/Models/SandboxTagging.h b/Engine/Plugins/Developer/Sandbox/SandboxedEditing/Source/SandboxedEditing/Private/Framework/Models/SandboxTagging.h
--- a/Engine/Plugins/Developer/Sandbox/SandboxedEditing/Source/SandboxedEditing/Private/Framework/Models/SandboxTagging.h
+++ b/Engine/Plugins/Developer/Sandbox/SandboxedEditing/Source/SandboxedEditing/Private/Framework/Models/SandboxTagging.h
@@ -14,6 +14,9 @@ constexpr FLazyName SandboxedEditingTag("SandboxedEditingTag");
 /** @return Whether the metadata is marked to belong to the Sandboxed Editing application. */
 bool IsOwnedBySandboxedEditing(const FFileSandboxCore_SandboxMetaData& InMetaData);
 
+/** @return Whether the metadata indicates that Sandboxed Editing can interact with this sandbox. */
+bool CanBeShownBySandboxedEditing(const FFileSandboxCore_SandboxMetaData& InMetaData);
+
 /** Modifies the metadata such that we know it belongs to Sandboxed Editing application. */
 void MarkOwnedBySandboxedEditing(FFileSandboxCore_SandboxMetaData& InMetaData);
 }
