File Sandbox: Prevent leaving with in-memory changes

#Synced-CL 49159615

diff --git a/Engine/Plugins/Developer/Sandbox/FileSandbox/Source/FileSandboxCore/Private/ISandboxManager.cpp b/Engine/Plugins/Developer/Sandbox/FileSandbox/Source/FileSandboxCore/Private/ISandboxManager.cpp
--- a/Engine/Plugins/Developer/Sandbox/FileSandbox/Source/FileSandboxCore/Private/ISandboxManager.cpp
+++ b/Engine/Plugins/Developer/Sandbox/FileSandbox/Source/FileSandboxCore/Private/ISandboxManager.cpp
@@ -10,6 +10,7 @@
 #include "Types/Manager/LoadSandboxByNameArgs.h"
 #include "Types/Manager/LoadSandboxError.h"
 #include "Types/Manager/SandboxCreationResult.h"
+#include "Utils/PackageSandboxUtils.h"
 #include "Utils/SandboxDirectoryUtils.h"
 
 namespace UE::FileSandboxCore
@@ -38,15 +39,21 @@ FDeleteSandboxResult ISandboxManager::DeleteNamedSandbox(const FDeleteSandboxByN
 	return DeleteSandbox(MoveTemp(*FullPath));
 }
 
-bool ISandboxManager::CanLeaveSandbox() const
+ELeaveSandboxErrorCode ISandboxManager::CanLeaveSandboxWithReason() const
 {
-	if (!HasActiveSandbox())
+	if (HasDirtyPackages())
 	{
-		return true;
+		return ELeaveSandboxErrorCode::InMemoryChanges;
 	}
 	
 	ISandboxLock* ActiveLock = GetActiveLock();
-	return !ActiveLock || ActiveLock->CanLeaveSandbox();
+	const bool bCanLeaveDueToLock = !ActiveLock || ActiveLock->CanLeaveSandbox();
+	if (!bCanLeaveDueToLock)
+	{
+		return ELeaveSandboxErrorCode::Locked;
+	}
+	
+	return ELeaveSandboxErrorCode::Success;
 }
 
 void ISandboxManager::EnumerateFileChangesByName(
diff --git a/Engine/Plugins/Developer/Sandbox/FileSandbox/Source/FileSandboxCore/Private/Sandbox/SandboxManager.cpp b/Engine/Plugins/Developer/Sandbox/FileSandbox/Source/FileSandboxCore/Private/Sandbox/SandboxManager.cpp
--- a/Engine/Plugins/Developer/Sandbox/FileSandbox/Source/FileSandboxCore/Private/Sandbox/SandboxManager.cpp
+++ b/Engine/Plugins/Developer/Sandbox/FileSandbox/Source/FileSandboxCore/Private/Sandbox/SandboxManager.cpp
@@ -121,9 +121,10 @@ FLeaveSandboxResult FSandboxManager::LeaveSandbox(const FLeaveSandboxArgs& InLea
 		return ELeaveSandboxErrorCode::Success;
 	}
 	
-	if (!CanLeaveSandbox())
+	if (const ELeaveSandboxErrorCode ErrorCode = CanLeaveSandboxWithReason()
+		; ErrorCode != ELeaveSandboxErrorCode::Success)
 	{
-		return ELeaveSandboxErrorCode::Locked;
+		return ErrorCode;
 	}
 	
 	OnPreSandboxShutdownDelegate.Broadcast(*ActiveSandboxData->Sandbox.Get());
diff --git a/Engine/Plugins/Developer/Sandbox/FileSandbox/Source/FileSandboxCore/Private/Utils/PackageSandboxUtils.cpp b/Engine/Plugins/Developer/Sandbox/FileSandbox/Source/FileSandboxCore/Private/Utils/PackageSandboxUtils.cpp
--- a/Engine/Plugins/Developer/Sandbox/FileSandbox/Source/FileSandboxCore/Private/Utils/PackageSandboxUtils.cpp
+++ b/Engine/Plugins/Developer/Sandbox/FileSandbox/Source/FileSandboxCore/Private/Utils/PackageSandboxUtils.cpp
@@ -55,7 +55,7 @@ bool ShouldReloadPersistentLevel(UPackage *PackageToReload)
 	return false;
 }
 	
-void PurgePackages(const TArray<FName>& InPackageNames)
+void PurgePackages(const TConstArrayView<FName>& InPackageNames)
 {
 	if (InPackageNames.Num() == 0)
 	{
@@ -143,11 +143,11 @@ void PurgePackages(const TArray<FName>& InPackageNames)
 
 }
 
-void HotReloadPackages(const TArray<FName>& InPackageNames)
+bool HotReloadPackages(const TConstArrayView<FName>& InPackageNames)
 {
 	if (InPackageNames.Num() == 0)
 	{
-		return;
+		return true;
 	}
 
 #if WITH_EDITOR
@@ -165,29 +165,40 @@ void HotReloadPackages(const TArray<FName>& InPackageNames)
 		}
 	}
 
-	bool bAddPersistentLevel = false;
-
 	// Find the packages in-memory to content hot-reload
 	TArray<UPackage*> ExistingPackages;
-	ExistingPackages.Reserve(InPackageNames.Num());
+	Algo::Transform(InPackageNames, ExistingPackages, [](FName PackageName)
+	{
+		return FindPackage(nullptr, *PackageName.ToString());
+	});
+	return HotReloadPackages(ExistingPackages);
+#else
+	return true;
+#endif
+}
 
-	for (const FName& PackageName : InPackageNames)
+bool HotReloadPackages(const TConstArrayView<UPackage*>& InPackages)
+{
+#if WITH_EDITOR
+	bool bAddPersistentLevel = false;
+	
+	TArray<UPackage*> FilteredPackages;
+	for (UPackage* Package : InPackages)
 	{
-		UPackage* ExistingPackage = FindPackage(nullptr, *PackageName.ToString());
-		if (ExistingPackage)
+		if (Package)
 		{
-			if (ExistingPackage->HasAnyPackageFlags(PKG_NewlyCreated))
+			if (Package->HasAnyPackageFlags(PKG_NewlyCreated))
 			{
-				ExistingPackage->ClearPackageFlags(PKG_NewlyCreated);
+				Package->ClearPackageFlags(PKG_NewlyCreated);
 			}
-			if (ExistingPackage->ContainsMap())
+			if (Package->ContainsMap())
 			{
-				bAddPersistentLevel = ShouldReloadPersistentLevel(ExistingPackage);
+				bAddPersistentLevel = ShouldReloadPersistentLevel(Package);
 			}
 
-			if (!ExistingPackage->ContainsMap() || ExistingPackage->IsDirty())
+			if (Package->ContainsMap() || Package->IsDirty())
 			{
-				ExistingPackages.Add(ExistingPackage);
+				FilteredPackages.Add(Package);
 			}
 		}
 	}
@@ -196,24 +207,36 @@ void HotReloadPackages(const TArray<FName>& InPackageNames)
 	if (CurrentWorld && bAddPersistentLevel)
 	{
 		ULevel* PersistentLevel = CurrentWorld->PersistentLevel;
-		if (PersistentLevel && !ExistingPackages.Contains(PersistentLevel->GetPackage()))
+		if (PersistentLevel && !FilteredPackages.Contains(PersistentLevel->GetPackage()))
 		{
-			ExistingPackages.Add(PersistentLevel->GetPackage());
+			FilteredPackages.Add(PersistentLevel->GetPackage());
 		}
 	}
 
-	if (ExistingPackages.Num() > 0)
+	if (FilteredPackages.Num() > 0)
 	{
 		FlushRenderingCommands();
 
 		FText ErrorMessage;
-		UPackageTools::ReloadPackages(ExistingPackages, ErrorMessage, UPackageTools::EReloadPackagesInteractionMode::AssumePositive);
+		UPackageTools::ReloadPackages(FilteredPackages, ErrorMessage, UPackageTools::EReloadPackagesInteractionMode::AssumePositive);
 
 		if (!ErrorMessage.IsEmpty())
 		{
 			FMessageDialog::Open(EAppMsgType::Ok, ErrorMessage);
+			return false;
 		}
 	}
 #endif
+	return true;
+}
+
+TArray<UPackage*> GetDirtyPackages()
+{
+	TArray<UPackage*> DirtyPackages;
+#if WITH_EDITOR
+	UEditorLoadingAndSavingUtils::GetDirtyMapPackages(DirtyPackages);
+	UEditorLoadingAndSavingUtils::GetDirtyContentPackages(DirtyPackages);
+#endif
+	return DirtyPackages;
 }
 }
diff --git a/Engine/Plugins/Developer/Sandbox/FileSandbox/Source/FileSandboxCore/Public/ISandboxManager.h b/Engine/Plugins/Developer/Sandbox/FileSandbox/Source/FileSandboxCore/Public/ISandboxManager.h
--- a/Engine/Plugins/Developer/Sandbox/FileSandbox/Source/FileSandboxCore/Public/ISandboxManager.h
+++ b/Engine/Plugins/Developer/Sandbox/FileSandbox/Source/FileSandboxCore/Public/ISandboxManager.h
@@ -104,9 +104,10 @@ public:
 	 */
 	virtual ISandboxLock* GetActiveLock() const = 0;
 	
-	/** @return Info about whether a LeaveSandbox call would be successful. */
-	UE_API bool CanLeaveSandbox() const;
-	
+	/** @return Whether a LeaveSandbox call would be successful. */
+	bool CanLeaveSandbox() const { return CanLeaveSandboxWithReason() == ELeaveSandboxErrorCode::Success; }
+	/** @return Why a LeaveSandbox call would be (un-)successful. */
+	UE_API ELeaveSandboxErrorCode CanLeaveSandboxWithReason() const;
 
 	/**
 	 * Gets the currently active global sandbox instance.
diff --git a/Engine/Plugins/Developer/Sandbox/FileSandbox/Source/FileSandboxCore/Public/Types/Manager/LeaveSandboxResult.h b/Engine/Plugins/Developer/Sandbox/FileSandbox/Source/FileSandboxCore/Public/Types/Manager/LeaveSandboxResult.h
--- a/Engine/Plugins/Developer/Sandbox/FileSandbox/Source/FileSandboxCore/Public/Types/Manager/LeaveSandboxResult.h
+++ b/Engine/Plugins/Developer/Sandbox/FileSandbox/Source/FileSandboxCore/Public/Types/Manager/LeaveSandboxResult.h
@@ -12,6 +12,14 @@ enum class ELeaveSandboxErrorCode : uint8
 	/** The sandbox was left successfully, or the engine was not in any sandbox. */
 	Success,
 	
+	/** 
+	 * There are packages with in-memory changes. Either discard them or save them out first. 
+	 * @see FileSandboxCore::GetDirtyPackages
+	 * @see FileSandboxCore::HotReloadPackages to revert changes.
+	 * @see FEditorFileUtils::SaveDirtyPackages to save all assets.
+	 */
+	InMemoryChanges,
+	
 	/** 
 	 * The engine is locked to this sandbox. Inspect the ISandboxLock for the reason.
 	 * @see ISandboxManager::GetActiveLock
@@ -22,7 +30,7 @@ enum class ELeaveSandboxErrorCode : uint8
 /** Describes the result of the leave operation, with a reason of why it failed. */
 struct FLeaveSandboxResult
 {
-	/** Describes the result of the delete operation, with a reason of why it failed. */
+	/** Describes the result of the leave operation, with a reason of why it failed. */
 	ELeaveSandboxErrorCode ErrorCode = ELeaveSandboxErrorCode::Success;
 	
 	FLeaveSandboxResult() = default;
diff --git a/Engine/Plugins/Developer/Sandbox/FileSandbox/Source/FileSandboxCore/Public/Utils/PackageSandboxUtils.h b/Engine/Plugins/Developer/Sandbox/FileSandbox/Source/FileSandboxCore/Public/Utils/PackageSandboxUtils.h
--- a/Engine/Plugins/Developer/Sandbox/FileSandbox/Source/FileSandboxCore/Public/Utils/PackageSandboxUtils.h
+++ b/Engine/Plugins/Developer/Sandbox/FileSandbox/Source/FileSandboxCore/Public/Utils/PackageSandboxUtils.h
@@ -4,13 +4,23 @@
 
 #include "Containers/Array.h"
 
+#define UE_API FILESANDBOXCORE_API
+
 class FName;
 
 namespace UE::FileSandboxCore
 {
 /** Handles cleaning up assets that no longer exist. */
-void PurgePackages(const TArray<FName>& InPackageNames);
+UE_API void PurgePackages(const TConstArrayView<FName>& InPackageNames);
 
 /** Reloads packages that were modified */
-void HotReloadPackages(const TArray<FName>& InPackageNames);
+UE_API bool HotReloadPackages(const TConstArrayView<FName>& InPackageNames);
+UE_API bool HotReloadPackages(const TConstArrayView<UPackage*>& InPackages);
+
+/** @return Packages with in-memory changes. */
+UE_API TArray<UPackage*> GetDirtyPackages();
+/** @return Whether there are packages with in-memory changes. */
+inline bool HasDirtyPackages() { return !GetDirtyPackages().IsEmpty(); }
 }
+
+#undef UE_API
\ No newline at end of file
