File Sandbox Core: Sandboxes now revert in-memory changes on leave instead of rejecting to leave the sandbox.

#Synced-CL 49159615

diff --git a/Engine/Plugins/Developer/Sandbox/FileSandbox/Source/FileSandboxCore/Private/ISandboxManager.cpp b/Engine/Plugins/Developer/Sandbox/FileSandbox/Source/FileSandboxCore/Private/ISandboxManager.cpp
--- a/Engine/Plugins/Developer/Sandbox/FileSandbox/Source/FileSandboxCore/Private/ISandboxManager.cpp
+++ b/Engine/Plugins/Developer/Sandbox/FileSandbox/Source/FileSandboxCore/Private/ISandboxManager.cpp
@@ -41,11 +41,6 @@ FDeleteSandboxResult ISandboxManager::DeleteNamedSandbox(const FDeleteSandboxByN
 
 ELeaveSandboxErrorCode ISandboxManager::CanLeaveSandboxWithReason() const
 {
-	if (HasDirtyPackages())
-	{
-		return ELeaveSandboxErrorCode::InMemoryChanges;
-	}
-	
 	ISandboxLock* ActiveLock = GetActiveLock();
 	const bool bCanLeaveDueToLock = !ActiveLock || ActiveLock->CanLeaveSandbox();
 	if (!bCanLeaveDueToLock)
diff --git a/Engine/Plugins/Developer/Sandbox/FileSandbox/Source/FileSandboxCore/Private/Sandbox/SandboxInstance.cpp b/Engine/Plugins/Developer/Sandbox/FileSandbox/Source/FileSandboxCore/Private/Sandbox/SandboxInstance.cpp
--- a/Engine/Plugins/Developer/Sandbox/FileSandbox/Source/FileSandboxCore/Private/Sandbox/SandboxInstance.cpp
+++ b/Engine/Plugins/Developer/Sandbox/FileSandbox/Source/FileSandboxCore/Private/Sandbox/SandboxInstance.cpp
@@ -11,6 +11,7 @@
 #include "Types/Manager/SandboxInitArgs.h"
 #include "Types/Sandbox/DiscardResult.h"
 #include "Types/SandboxFileChange.h"
+#include "UObject/Package.h"
 #include "Utils/PackageSandboxUtils.h"
 #include "Utils/SandboxFileUtils.h"
 #include "Utils/SandboxDirectoryUtils.h"
@@ -32,6 +33,21 @@ static void TrackRevertedFileChangeInManifest(FFileSandboxCore_ManifestData& Out
 	default: break;
 	}
 }
+
+/** 
+ * Adds packages with in-memory changes. 
+ * 
+ * These packages must be reverted in order to discard their in-memory changed.
+ * Revert effectively happens by doing a hot reload.
+ */
+static void AppendPackagesWithInMemoryChanges(TArray<FName>& OutPackagesPendingHotReload)
+{
+	const TArray<UPackage*> DirtyPackages = GetDirtyPackages();
+	Algo::TransformIf(DirtyPackages, OutPackagesPendingHotReload, 
+		[&](UPackage* Package) { return !OutPackagesPendingHotReload.Contains(Package->GetFName()); },
+		[&](UPackage* Package) { return Package->GetFName(); }
+		);
+}
 }
 	
 TUniquePtr<FSandboxInstance> FSandboxInstance::CreateNewSandbox(
@@ -91,6 +107,8 @@ FSandboxInstance::~FSandboxInstance()
 
 	if (!IsEngineExitRequested())
 	{
+		Private::AppendPackagesWithInMemoryChanges(PackagesPendingHotReload);
+		
 		// Hot reload after unregistering from most delegates to prevent events triggered by hot-reloading (such as asset deleted) to be recorded as transaction.
 		PurgePackages(PackagesPendingPurge);
 		HotReloadPackages(PackagesPendingHotReload);
@@ -144,7 +162,8 @@ FDiscardResult FSandboxInstance::DiscardSandbox(const FDiscardArgs& InArgs)
 	
 	ManifestContent.Empty();
 	SaveManifest();
-
+	
+	Private::AppendPackagesWithInMemoryChanges(Result.PackagesPendingHotReload);
 	return Result;
 }
 
diff --git a/Engine/Plugins/Developer/Sandbox/FileSandbox/Source/FileSandboxCore/Public/Types/Manager/LeaveSandboxResult.h b/Engine/Plugins/Developer/Sandbox/FileSandbox/Source/FileSandboxCore/Public/Types/Manager/LeaveSandboxResult.h
--- a/Engine/Plugins/Developer/Sandbox/FileSandbox/Source/FileSandboxCore/Public/Types/Manager/LeaveSandboxResult.h
+++ b/Engine/Plugins/Developer/Sandbox/FileSandbox/Source/FileSandboxCore/Public/Types/Manager/LeaveSandboxResult.h
@@ -12,14 +12,6 @@ enum class ELeaveSandboxErrorCode : uint8
 	/** The sandbox was left successfully, or the engine was not in any sandbox. */
 	Success,
 	
-	/** 
-	 * There are packages with in-memory changes. Either discard them or save them out first. 
-	 * @see FileSandboxCore::GetDirtyPackages
-	 * @see FileSandboxCore::HotReloadPackages to revert changes.
-	 * @see FEditorFileUtils::SaveDirtyPackages to save all assets.
-	 */
-	InMemoryChanges,
-	
 	/** 
 	 * The engine is locked to this sandbox. Inspect the ISandboxLock for the reason.
 	 * @see ISandboxManager::GetActiveLock
diff --git a/Engine/Plugins/Developer/Sandbox/SandboxedEditing/Source/SandboxedEditing/Private/Framework/Models/SandboxSystemModel.cpp b/Engine/Plugins/Developer/Sandbox/SandboxedEditing/Source/SandboxedEditing/Private/Framework/Models/SandboxSystemModel.cpp
--- a/Engine/Plugins/Developer/Sandbox/SandboxedEditing/Source/SandboxedEditing/Private/Framework/Models/SandboxSystemModel.cpp
+++ b/Engine/Plugins/Developer/Sandbox/SandboxedEditing/Source/SandboxedEditing/Private/Framework/Models/SandboxSystemModel.cpp
@@ -153,8 +153,7 @@ bool FSandboxSystemModel::IsAllowedToLeaveSandbox(FText* OutReason) const
 	}
 	
 	const ELeaveSandboxErrorCode LeaveReason = SandboxManager.CanLeaveSandboxWithReason();
-	return LeaveReason == ELeaveSandboxErrorCode::InMemoryChanges  // The user is able to leave the sandbox - the app will handle the dirty packages.
-		|| LeaveReason == ELeaveSandboxErrorCode::Success;
+	return LeaveReason == ELeaveSandboxErrorCode::Success;
 }
 #undef SET_REASON
 
