File Sandbox Core: Expose IPersistFeedback so API users know what files were able to be persisted.

#Synced-CL 49159615

diff --git a/Engine/Plugins/Developer/Concert/ConcertSync/ConcertSyncClient/Source/ConcertSyncClient/Private/Sandbox/ConcertPackageSandbox.cpp b/Engine/Plugins/Developer/Concert/ConcertSync/ConcertSyncClient/Source/ConcertSyncClient/Private/Sandbox/ConcertPackageSandbox.cpp
--- a/Engine/Plugins/Developer/Concert/ConcertSync/ConcertSyncClient/Source/ConcertSyncClient/Private/Sandbox/ConcertPackageSandbox.cpp
+++ b/Engine/Plugins/Developer/Concert/ConcertSync/ConcertSyncClient/Source/ConcertSyncClient/Private/Sandbox/ConcertPackageSandbox.cpp
@@ -9,6 +9,7 @@
 #include "IFileSandboxCoreModule.h"
 #include "ISandboxInstance.h"
 #include "ISandboxManager.h"
+#include "Interface/Feedback/AppendTextErrorHandler.h"
 #include "Misc/PackageName.h"
 #include "Sandbox/ConcertSandboxLock.h"
 #include "Types/Manager/LoadSandboxByNameArgs.h"
@@ -114,9 +115,10 @@ FPersistResult FConcertPackageSandbox::PersistSandbox(TArrayView<const FString>
 	ISandboxManager& SandboxManager = IFileSandboxCoreModule::Get().GetSandboxManager();
 	if (ISandboxInstance* Sandbox = SandboxManager.GetActiveSandboxInstance())
 	{
-		const FPersistArgs Args { InFiles, InParams.bShouldMakeWritableIfNoSourceControl };
+		FAppendTextPersistFeedback ErrorAsTextFeedback;
+		const FPersistArgs Args { InFiles, InParams.bShouldMakeWritableIfNoSourceControl, &ErrorAsTextFeedback };
 		const FPersistResult Result = Sandbox->PersistSandbox(Args);
-		return ::FPersistResult(Result.FailureReasons, static_cast<::EPersistStatus>(Result.PersistStatus));
+		return ::FPersistResult(ErrorAsTextFeedback.AllErrors, static_cast<::EPersistStatus>(Result.PersistStatus));
 	}
 
 	return ::FPersistResult { .PersistStatus = ::EPersistStatus::Failure };
diff --git a/Engine/Plugins/Developer/Sandbox/FileSandbox/Source/FileSandboxCore/Private/Sandbox/Persist/AggregatePersistErrorHandler.h b/Engine/Plugins/Developer/Sandbox/FileSandbox/Source/FileSandboxCore/Private/Sandbox/Persist/AggregatePersistErrorHandler.h
deleted file mode 100644
diff --git a/Engine/Plugins/Developer/Sandbox/FileSandbox/Source/FileSandboxCore/Private/Sandbox/Persist/AppendTextErrorHandler.h b/Engine/Plugins/Developer/Sandbox/FileSandbox/Source/FileSandboxCore/Private/Sandbox/Persist/AppendTextErrorHandler.h
deleted file mode 100644
diff --git a/Engine/Plugins/Developer/Sandbox/FileSandbox/Source/FileSandboxCore/Private/Sandbox/Persist/IPersistErrorHandler.h b/Engine/Plugins/Developer/Sandbox/FileSandbox/Source/FileSandboxCore/Private/Sandbox/Persist/IPersistErrorHandler.h
deleted file mode 100644
diff --git a/Engine/Plugins/Developer/Sandbox/FileSandbox/Source/FileSandboxCore/Private/Sandbox/Persist/LogPersistErrorHandler.h b/Engine/Plugins/Developer/Sandbox/FileSandbox/Source/FileSandboxCore/Private/Sandbox/Persist/LogPersistErrorHandler.h
deleted file mode 100644
diff --git a/Engine/Plugins/Developer/Sandbox/FileSandbox/Source/FileSandboxCore/Private/Sandbox/Persist/LogPersistFeedback.h b/Engine/Plugins/Developer/Sandbox/FileSandbox/Source/FileSandboxCore/Private/Sandbox/Persist/LogPersistFeedback.h
new file mode 100644
--- /dev/null
+++ b/Engine/Plugins/Developer/Sandbox/FileSandbox/Source/FileSandboxCore/Private/Sandbox/Persist/LogPersistFeedback.h
@@ -0,0 +1,58 @@
+﻿// Copyright Epic Games, Inc. All Rights Reserved.
+
+#pragma once 
+
+#include "Interface/IPersistFeedback.h"
+#include "LogFileSandbox.h"
+#include "Logging/LogMacros.h"
+
+namespace UE::FileSandboxCore
+{
+class FLogPersistFeedback : public IPersistFeedback
+{
+public:
+	
+	//~ Begin IPersistFeedback Interface
+	virtual void StartFile(const TCHAR* InFilename) override
+	{
+		UE_LOG(LogFileSandbox, Verbose, TEXT("Starting persist for file '%s'"), InFilename);
+	}
+	virtual void HandleSuccess(const TCHAR* InFilename) override
+	{
+		UE_LOG(LogFileSandbox, Verbose, TEXT("Persisted file '%s'"), InFilename);
+	}
+	virtual void HandleError_CheckoutNotAllowed(const TCHAR* InFilename) override
+	{
+		UE_LOG(LogFileSandbox, Warning, TEXT("Checkout of file '%s' from revision control was not allowed when persisting sandbox state!"), InFilename);
+	}
+	virtual void HandleError_Checkout(const TCHAR* InFilename) override
+	{
+		UE_LOG(LogFileSandbox, Warning, TEXT("Failed to checkout file '%s' from revision control when persisting sandbox state!"), InFilename);
+	}
+	virtual void HandleError_Revert(const TCHAR* InFilename) override
+	{
+		UE_LOG(LogFileSandbox, Warning, TEXT("Failed to revert file '%s' from revision control when persisting sandbox state!"), InFilename);
+	}
+	virtual void HandleError_MarkForAdd(const TCHAR* InFilename) override
+	{
+		UE_LOG(LogFileSandbox, Warning, TEXT("Failed to add file '%s' to revision control when persisting sandbox state!"), InFilename);
+	}
+	virtual void HandleError_DeleteSCC(const TCHAR* InFilename) override
+	{
+		UE_LOG(LogFileSandbox, Warning, TEXT("Failed to delete file '%s' from revision control when persisting sandbox state!"), InFilename);
+	}
+	virtual void HandleError_MakeWritable(const TCHAR* InFilename) override
+	{
+		UE_LOG(LogFileSandbox, Warning, TEXT("Failed to make '%s' writeable using file system when persisting sandbox state!"), InFilename);
+	}
+	virtual void HandleError_MoveFile(const TCHAR* InToFilename, const TCHAR* InFromFilename) override
+	{
+		UE_LOG(LogFileSandbox, Warning, TEXT("Failed to move file '%s' to '%s' using file system when persisting sandbox state!"), InFromFilename, InToFilename);
+	}
+	virtual void HandleError_DeleteFile(const TCHAR* InFilename) override
+	{
+		UE_LOG(LogFileSandbox, Warning, TEXT("Failed to delete file '%s' using file system when persisting sandbox state!"), InFilename);
+	}
+	//~ End IPersistFeedback Interface
+};
+}
diff --git a/Engine/Plugins/Developer/Sandbox/FileSandbox/Source/FileSandboxCore/Private/Sandbox/Persist/PersistUtils.cpp b/Engine/Plugins/Developer/Sandbox/FileSandbox/Source/FileSandboxCore/Private/Sandbox/Persist/PersistUtils.cpp
--- a/Engine/Plugins/Developer/Sandbox/FileSandbox/Source/FileSandboxCore/Private/Sandbox/Persist/PersistUtils.cpp
+++ b/Engine/Plugins/Developer/Sandbox/FileSandbox/Source/FileSandboxCore/Private/Sandbox/Persist/PersistUtils.cpp
@@ -2,12 +2,12 @@
 
 #include "PersistUtils.h"
 
-#include "IPersistErrorHandler.h"
+#include "GenericPlatform/GenericPlatformFile.h"
 #include "ISourceControlProvider.h"
 #include "ISourceControlState.h"
-#include "SourceControlOperations.h"
-#include "GenericPlatform/GenericPlatformFile.h"
+#include "Interface/IPersistFeedback.h"
 #include "Misc/Paths.h"
+#include "SourceControlOperations.h"
 #include "Utils/AssetUtils.h"
 
 namespace UE::FileSandboxCore
@@ -15,7 +15,7 @@ namespace UE::FileSandboxCore
 namespace PersistDetail
 {
 static bool TryCheckout(
-	ISourceControlProvider& InSourceControlProvider, IPersistErrorHandler& InErrorHandler,
+	ISourceControlProvider& InSourceControlProvider, IPersistFeedback& InErrorHandler,
 	const TCHAR* InToFilename, const FSourceControlStateRef& ToFileSCCState
 	)
 {
@@ -44,7 +44,7 @@ static bool TryCheckout(
 }
 
 bool AddFileWithSCC(
-	IPlatformFile& InPlatformFile, ISourceControlProvider& InSourceControlProvider, IPersistErrorHandler& InErrorHandler,
+	IPlatformFile& InPlatformFile, ISourceControlProvider& InSourceControlProvider, IPersistFeedback& InErrorHandler,
 	const TCHAR* InToFilename, const TCHAR* InFromFilename,
 	EPersistFlags InFlags
 	)
@@ -82,7 +82,7 @@ bool AddFileWithSCC(
 }
 
 bool EditFileWithSCC(
-	IPlatformFile& InPlatformFile, ISourceControlProvider& InSourceControlProvider, IPersistErrorHandler& InErrorHandler,
+	IPlatformFile& InPlatformFile, ISourceControlProvider& InSourceControlProvider, IPersistFeedback& InErrorHandler,
 	const TCHAR* InToFilename, const TCHAR* InFromFilename, 
 	bool bMarkWriteable, 
 	EPersistFlags InFlags
@@ -113,7 +113,7 @@ bool EditFileWithSCC(
 }
 
 bool DeleteFileWithSCC(
-	IPlatformFile& InPlatformFile, ISourceControlProvider& InSourceControlProvider, IPersistErrorHandler& InErrorHandler, const TCHAR* InFilename
+	IPlatformFile& InPlatformFile, ISourceControlProvider& InSourceControlProvider, IPersistFeedback& InErrorHandler, const TCHAR* InFilename
 	)
 {
 	// If this file maps to a package then we need to flush its linker so that we can remove the file from disk
@@ -153,7 +153,7 @@ bool DeleteFileWithSCC(
 
 namespace PersistDetail
 {
-static bool CopyFile(IPlatformFile& InPlatformFile, IPersistErrorHandler& InErrorHandler, const TCHAR* InToFilename, const TCHAR* InFromFilename)
+static bool CopyFile(IPlatformFile& InPlatformFile, IPersistFeedback& InErrorHandler, const TCHAR* InToFilename, const TCHAR* InFromFilename)
 {
 	if (!InPlatformFile.CopyFile(InToFilename, InFromFilename, EPlatformFileRead::None, EPlatformFileWrite::AttemptDeleteAndCreate))
 	{
@@ -163,7 +163,7 @@ static bool CopyFile(IPlatformFile& InPlatformFile, IPersistErrorHandler& InErro
 	return true;
 }
 
-static bool MoveFile(IPlatformFile& InPlatformFile, IPersistErrorHandler& InErrorHandler, const TCHAR* InToFilename, const TCHAR* InFromFilename)
+static bool MoveFile(IPlatformFile& InPlatformFile, IPersistFeedback& InErrorHandler, const TCHAR* InToFilename, const TCHAR* InFromFilename)
 {
 	bool bSuccess = true;
 		
@@ -187,7 +187,7 @@ static bool MoveFile(IPlatformFile& InPlatformFile, IPersistErrorHandler& InErro
 }
 
 bool MoveSandboxFileToNonSandbox(
-	IPlatformFile& InPlatformFile, IPersistErrorHandler& InErrorHandler,
+	IPlatformFile& InPlatformFile, IPersistFeedback& InErrorHandler,
 	const TCHAR* InToFilename, const TCHAR* InFromFilename,
 	EPersistFlags InFlags
 	)
diff --git a/Engine/Plugins/Developer/Sandbox/FileSandbox/Source/FileSandboxCore/Private/Sandbox/Persist/PersistUtils.h b/Engine/Plugins/Developer/Sandbox/FileSandbox/Source/FileSandboxCore/Private/Sandbox/Persist/PersistUtils.h
--- a/Engine/Plugins/Developer/Sandbox/FileSandbox/Source/FileSandboxCore/Private/Sandbox/Persist/PersistUtils.h
+++ b/Engine/Plugins/Developer/Sandbox/FileSandbox/Source/FileSandboxCore/Private/Sandbox/Persist/PersistUtils.h
@@ -11,7 +11,7 @@ class ISourceControlProvider;
 namespace UE::FileSandboxCore
 {
 class FSandboxedPlatformFilePath;
-class IPersistErrorHandler;
+class IPersistFeedback;
 enum class ESandboxFileChange : uint8;
 struct FPersistArgs;
 
@@ -29,20 +29,20 @@ ENUM_CLASS_FLAGS(EPersistFlags);
 
 /** Persists a file that is missing from the non-sandbox but exists in the sandbox. Registers the changes with Source Control if enabled. */
 bool AddFileWithSCC(
-	IPlatformFile& InPlatformFile, ISourceControlProvider& InSourceControlProvider, IPersistErrorHandler& InErrorHandler,
+	IPlatformFile& InPlatformFile, ISourceControlProvider& InSourceControlProvider, IPersistFeedback& InErrorHandler,
 	const TCHAR* InToFilename, const TCHAR* InFromFilename, 
 	EPersistFlags InFlags = EPersistFlags::None
 	);
 /** Persists a file that exists in both the non-sandbox and sandbox. Registers the changes with Source Control if enabled. */
 bool EditFileWithSCC(
-	IPlatformFile& InPlatformFile, ISourceControlProvider& InSourceControlProvider, IPersistErrorHandler& InErrorHandler,
+	IPlatformFile& InPlatformFile, ISourceControlProvider& InSourceControlProvider, IPersistFeedback& InErrorHandler,
 	const TCHAR* InToFilename, const TCHAR* InFromFilename, 
 	bool bMarkWriteable, 
 	EPersistFlags InFlags = EPersistFlags::None
 	);
 /** Persists a file that exists in non-sandbox but was removed in the sandbox. Registers the changes with Source Control if enabled. */
 bool DeleteFileWithSCC(
-	IPlatformFile& InPlatformFile, ISourceControlProvider& InSourceControlProvider, IPersistErrorHandler& InErrorHandler,
+	IPlatformFile& InPlatformFile, ISourceControlProvider& InSourceControlProvider, IPersistFeedback& InErrorHandler,
 	const TCHAR* InFilename
 	);
 
@@ -53,7 +53,7 @@ bool DeleteFileWithSCC(
  * Depending on the flags InFromFilename stays in the sandbox, and continues to be tracked as change, or is removed.
  */
 bool MoveSandboxFileToNonSandbox(
-	IPlatformFile& PlatformFile, IPersistErrorHandler& InErrorHandler, 
+	IPlatformFile& PlatformFile, IPersistFeedback& InErrorHandler, 
 	const TCHAR* InToFilename, const TCHAR* InFromFilename, 
 	EPersistFlags InFlags = EPersistFlags::None
 	);
diff --git a/Engine/Plugins/Developer/Sandbox/FileSandbox/Source/FileSandboxCore/Private/Sandbox/Platform/SandboxPlatformFile.cpp b/Engine/Plugins/Developer/Sandbox/FileSandbox/Source/FileSandboxCore/Private/Sandbox/Platform/SandboxPlatformFile.cpp
--- a/Engine/Plugins/Developer/Sandbox/FileSandbox/Source/FileSandboxCore/Private/Sandbox/Platform/SandboxPlatformFile.cpp
+++ b/Engine/Plugins/Developer/Sandbox/FileSandbox/Source/FileSandboxCore/Private/Sandbox/Platform/SandboxPlatformFile.cpp
@@ -2,13 +2,13 @@
 
 #include "SandboxPlatformFile.h"
 
-#include "Algo/ForEach.h"
 #include "Containers/UnrealString.h"
 #include "Engine/GameEngine.h"
 #include "Engine/World.h"
 #include "GenericPlatform/GenericPlatformFile.h"
 #include "HAL/PlatformFileManager.h"
 #include "ISourceControlProvider.h"
+#include "Interface/IPersistFeedback.h"
 #include "JsonObjectConverter.h"
 #include "LogFileSandbox.h"
 #include "Logging/LogVerbosity.h"
@@ -18,19 +18,17 @@
 #include "Misc/ScopeLock.h"
 #include "Modules/ModuleManager.h"
 #include "RenderingThread.h"
+#include "Sandbox/Persist/DiffUtils.h"
 #include "Sandbox/Persist/DiscardUtils.h"
 #include "Sandbox/Persist/PersistUtils.h"
 #include "SandboxMountPoint.h"
-#include "Sandbox/Persist/DiffUtils.h"
 #include "Types/EBreakBehavior.h"
 #include "Types/Sandbox/PersistArgs.h"
-#include "Types/Sandbox/PersistResult.h"
 #include "Types/SandboxFileChange.h"
-#include "UObject/LinkerLoad.h"
 #include "UObject/Package.h"
 #include "Utils/AssetUtils.h"
-#include "Utils/SandboxFileUtils.h"
 #include "Utils/FileChange/FileChange.h"
+#include "Utils/SandboxFileUtils.h"
 
 #if UE_SANDBOX_WITH_DIRECTORY_WATCHER
 #include "Utils/Watcher/DirectoryWatcherUtils.h"
@@ -681,7 +679,7 @@ FString FSandboxPlatformFile::ConvertToAbsolutePathForExternalAppForWrite(const
 
 ESandboxFileChange FSandboxPlatformFile::PersistItem(
 	const FSandboxedPlatformFilePath& InFile,
-	ISourceControlProvider& InSourceControlProvider, const FPersistArgs& InParams, IPersistErrorHandler& InErrorHandler
+	ISourceControlProvider& InSourceControlProvider, const FPersistArgs& InParams, IPersistFeedback& InErrorHandler
 	)
 {
 	if (IsPathDeleted(InFile))
@@ -709,22 +707,24 @@ bool FSandboxPlatformFile::PersistSandbox(
 	ISourceControlProvider& InSourceControlProvider,
 	const FPersistArgs& InParams, 
 	TFunctionRef<void(const FSandboxedPlatformFilePath&, ESandboxFileChange)> InOnPersistFile,
-	IPersistErrorHandler& InErrorHandler
+	IPersistFeedback& InPersistFeedback
 	)
 {
-	FPersistResult OutResult;
 	// We need to disable the sandbox while we do this
 	bool bFullOperationSuccess = true;
 	TGuardValue<TAtomic<bool>, bool> DisableSandboxGuard(bSandboxEnabled, false);
 	for (const FString& File : InParams.Files)
 	{
+		InPersistFeedback.StartFile(*File);
+		
 		const FSandboxedPlatformFilePath FilePath = ToSandboxPath(File, true);
-		const ESandboxFileChange ActionPerformed = PersistItem(FilePath, InSourceControlProvider, InParams, InErrorHandler);
+		const ESandboxFileChange ActionPerformed = PersistItem(FilePath, InSourceControlProvider, InParams, InPersistFeedback);
 		const bool bSuccess = ActionPerformed != ESandboxFileChange::None;
 		
 		// if the file operation was successful, mark the file as persisted
 		if (bSuccess)
 		{
+			InPersistFeedback.HandleSuccess(*File);
 			InOnPersistFile(FilePath, ActionPerformed);
 		}
 
diff --git a/Engine/Plugins/Developer/Sandbox/FileSandbox/Source/FileSandboxCore/Private/Sandbox/Platform/SandboxPlatformFile.h b/Engine/Plugins/Developer/Sandbox/FileSandbox/Source/FileSandboxCore/Private/Sandbox/Platform/SandboxPlatformFile.h
--- a/Engine/Plugins/Developer/Sandbox/FileSandbox/Source/FileSandboxCore/Private/Sandbox/Platform/SandboxPlatformFile.h
+++ b/Engine/Plugins/Developer/Sandbox/FileSandbox/Source/FileSandboxCore/Private/Sandbox/Platform/SandboxPlatformFile.h
@@ -17,7 +17,7 @@ class ISourceControlProvider;
 
 namespace UE::FileSandboxCore
 {
-class IPersistErrorHandler;
+class IPersistFeedback;
 enum class EFileChangeAction : uint8;
 enum class ESandboxFileChange : uint8;
 struct FDiscardSandboxArgs;
@@ -78,13 +78,13 @@ public:
 	 *	@param InSourceControlProvider The source control provider to register the changes with
 	 *	@param InParams	Parameters controlling the persist options specified by the user.
 	 *	@param InOnPersistFile Invoked for each file successfully persisted.
-	 *	@param InErrorHandler Handles any errors that occur during persist
+	 *	@param InPersistFeedback Handles any errors that occur during persist
 	 *	@return Whether the persist operation succeeded completely. False if there were errors.
 	 */
 	bool PersistSandbox(
 		ISourceControlProvider& InSourceControlProvider, const FPersistArgs& InParams, 
 		TFunctionRef<void(const FSandboxedPlatformFilePath&, ESandboxFileChange)> InOnPersistFile, 
-		IPersistErrorHandler& InErrorHandler
+		IPersistFeedback& InPersistFeedback
 		);
 
 	/** 
@@ -158,7 +158,7 @@ private:
 	/** Attempts to persist the file. */
 	ESandboxFileChange PersistItem(
 		const FSandboxedPlatformFilePath& InFile, ISourceControlProvider& InSourceControlProvider, 
-		const FPersistArgs& InParams, IPersistErrorHandler& InErrorHandler
+		const FPersistArgs& InParams, IPersistFeedback& InErrorHandler
 		);
 
 	FString GetSandboxRootPath() const;
diff --git a/Engine/Plugins/Developer/Sandbox/FileSandbox/Source/FileSandboxCore/Private/Sandbox/SandboxInstance.cpp b/Engine/Plugins/Developer/Sandbox/FileSandbox/Source/FileSandboxCore/Private/Sandbox/SandboxInstance.cpp
--- a/Engine/Plugins/Developer/Sandbox/FileSandbox/Source/FileSandboxCore/Private/Sandbox/SandboxInstance.cpp
+++ b/Engine/Plugins/Developer/Sandbox/FileSandbox/Source/FileSandboxCore/Private/Sandbox/SandboxInstance.cpp
@@ -2,21 +2,18 @@
 
 #include "SandboxInstance.h"
 
-#include "HAL/FileManager.h"
+#include "Interface/Feedback/AggregatePersistFeedback.h"
 #include "Misc/PackageName.h"
-#include "Persist/AggregatePersistErrorHandler.h"
-#include "Persist/AppendTextErrorHandler.h"
 #include "Persist/DiffUtils.h"
-#include "Persist/LogPersistErrorHandler.h"
+#include "Persist/LogPersistFeedback.h"
 #include "Types/Manager/SandboxInitArgs.h"
 #include "Types/Sandbox/DiscardResult.h"
 #include "Types/SandboxFileChange.h"
 #include "UObject/Package.h"
-#include "Utils/PackageSandboxUtils.h"
-#include "Utils/SandboxFileUtils.h"
-#include "Utils/SandboxDirectoryUtils.h"
 #include "Utils/AssetUtils.h"
 #include "Utils/FileChange/FileChange.h"
+#include "Utils/PackageSandboxUtils.h"
+#include "Utils/SandboxDirectoryUtils.h"
 #include "Utils/SandboxFileUtils.h"
 
 namespace UE::FileSandboxCore
@@ -136,9 +133,8 @@ FPersistResult FSandboxInstance::PersistSandbox(const FPersistArgs& InArgs)
 		}
 	};
 	
-	FAppendTextPersistErrorHandler ErrorsAsReasons;
-	FLogPersistErrorHandler ErrorLogger;
-	FAggregatePersistErrorHandler Aggregate({ &ErrorsAsReasons, &ErrorLogger });
+	FLogPersistFeedback ErrorLogger;
+	FAggregatePersistFeedback Aggregate({ InArgs.Feedback, &ErrorLogger });
 	const bool bFullSuccess = PlatformFile->PersistSandbox(
 		SourceControlSandbox.GetProxiedSourceControl(), InArgs, HandleFileOp, Aggregate
 		);
@@ -150,7 +146,6 @@ FPersistResult FSandboxInstance::PersistSandbox(const FPersistArgs& InArgs)
 
 	return FPersistResult
 	{ 
-		MoveTemp(ErrorsAsReasons.AllErrors),
 		bFullSuccess ? EPersistStatus::Success : EPersistStatus::Failure 
 	};
 }
diff --git a/Engine/Plugins/Developer/Sandbox/FileSandbox/Source/FileSandboxCore/Public/Interface/Feedback/AggregatePersistFeedback.h b/Engine/Plugins/Developer/Sandbox/FileSandbox/Source/FileSandboxCore/Public/Interface/Feedback/AggregatePersistFeedback.h
new file mode 100644
--- /dev/null
+++ b/Engine/Plugins/Developer/Sandbox/FileSandbox/Source/FileSandboxCore/Public/Interface/Feedback/AggregatePersistFeedback.h
@@ -0,0 +1,127 @@
+﻿// Copyright Epic Games, Inc. All Rights Reserved.
+
+#pragma once 
+
+#include "Algo/ForEach.h"
+#include "Containers/Array.h"
+#include "Interface/IPersistFeedback.h"
+
+namespace UE::FileSandboxCore
+{
+/** Forwards the errors to other error handlers. */
+class FAggregatePersistFeedback : public IPersistFeedback
+{
+public:
+	
+	explicit FAggregatePersistFeedback(TArray<IPersistFeedback*> InErrorHandlers) : ErrorHandlers(MoveTemp(InErrorHandlers)) {}
+	explicit FAggregatePersistFeedback(IPersistFeedback& InErrorHandler) : ErrorHandlers({ &InErrorHandler }) {}
+	
+	//~ Begin IPersistFeedback Interface
+	virtual void StartFile(const TCHAR* InFilename) override
+	{
+		Algo::ForEach(ErrorHandlers, [&](IPersistFeedback* Handler)
+		{
+			if (Handler)
+			{
+				Handler->StartFile(InFilename);
+			}
+		});
+	}
+	virtual void HandleSuccess(const TCHAR* InFilename) override
+	{
+		Algo::ForEach(ErrorHandlers, [&](IPersistFeedback* Handler)
+		{
+			if (Handler)
+			{
+				Handler->HandleSuccess(InFilename);
+			}
+		});
+	}
+	virtual void HandleError_CheckoutNotAllowed(const TCHAR* InFilename) override
+	{
+		Algo::ForEach(ErrorHandlers, [&](IPersistFeedback* Handler)
+		{
+			if (Handler)
+			{
+				Handler->HandleError_CheckoutNotAllowed(InFilename);
+			}
+		});
+	}
+	virtual void HandleError_Checkout(const TCHAR* InFilename) override
+	{
+		Algo::ForEach(ErrorHandlers, [&](IPersistFeedback* Handler)
+		{
+			if (Handler)
+			{
+				Handler->HandleError_Checkout(InFilename);
+			}
+		});
+	}
+	virtual void HandleError_Revert(const TCHAR* InFilename) override
+	{
+		Algo::ForEach(ErrorHandlers, [&](IPersistFeedback* Handler)
+		{
+			if (Handler)
+			{
+				Handler->HandleError_Revert(InFilename);
+			}
+		});
+	}
+	virtual void HandleError_MarkForAdd(const TCHAR* InFilename) override
+	{
+		Algo::ForEach(ErrorHandlers, [&](IPersistFeedback* Handler)
+		{
+			if (Handler)
+			{
+				Handler->HandleError_MarkForAdd(InFilename);
+			}
+		});
+	}
+	virtual void HandleError_DeleteSCC(const TCHAR* InFilename) override
+	{
+		Algo::ForEach(ErrorHandlers, [&](IPersistFeedback* Handler)
+		{
+			if (Handler)
+			{
+				Handler->HandleError_DeleteSCC(InFilename);
+			}
+		});
+	}
+	virtual void HandleError_MakeWritable(const TCHAR* InFilename) override
+	{
+		Algo::ForEach(ErrorHandlers, [&](IPersistFeedback* Handler)
+		{
+			if (Handler)
+			{
+				Handler->HandleError_MakeWritable(InFilename);
+			}
+		});
+	}
+	virtual void HandleError_MoveFile(const TCHAR* InToFilename, const TCHAR* InFromFilename) override
+	{
+		Algo::ForEach(ErrorHandlers, [&](IPersistFeedback* Handler)
+		{
+			if (Handler)
+			{
+				Handler->HandleError_MoveFile(InToFilename, InFromFilename);
+			}
+		});
+	}
+	virtual void HandleError_DeleteFile(const TCHAR* InFilename) override
+	{
+		Algo::ForEach(ErrorHandlers, [&](IPersistFeedback* Handler)
+		{
+			if (Handler)
+			{
+				Handler->HandleError_DeleteFile(InFilename);
+			}
+		});
+	}
+	//~ End IPersistFeedback Interface
+	
+private:
+	
+	/** The error handlers to forward the errors to. */
+	TArray<IPersistFeedback*> ErrorHandlers;
+};
+}
diff --git a/Engine/Plugins/Developer/Sandbox/FileSandbox/Source/FileSandboxCore/Public/Interface/Feedback/AppendTextErrorHandler.h b/Engine/Plugins/Developer/Sandbox/FileSandbox/Source/FileSandboxCore/Public/Interface/Feedback/AppendTextErrorHandler.h
new file mode 100644
--- /dev/null
+++ b/Engine/Plugins/Developer/Sandbox/FileSandbox/Source/FileSandboxCore/Public/Interface/Feedback/AppendTextErrorHandler.h
@@ -0,0 +1,83 @@
+﻿// Copyright Epic Games, Inc. All Rights Reserved.
+
+#pragma once 
+
+#include "Interface/IPersistFeedback.h"
+#include "Internationalization/Internationalization.h"
+#include "Internationalization/Text.h"
+
+#define LOCTEXT_NAMESPACE "FAppendTextPersistErrorHandler"
+
+namespace UE::FileSandboxCore
+{
+class FAppendTextPersistFeedback : public IPersistFeedback
+{
+public:
+	
+	TArray<FText> AllErrors;
+	
+	//~ Begin IPersistFeedback Interface
+	virtual void StartFile(const TCHAR* InFilename) override {}
+	virtual void HandleSuccess(const TCHAR* InFilename) override {}
+	virtual void HandleError_CheckoutNotAllowed(const TCHAR* InFilename) override
+	{
+		AllErrors.Emplace(FText::Format(
+			LOCTEXT("CheckoutNotAllowed", "Checkout of file '{0}' from revision control was not allowed when persisting sandbox state!"),
+			FText::AsCultureInvariant(InFilename)
+			));
+	}
+	virtual void HandleError_Checkout(const TCHAR* InFilename) override
+	{
+		AllErrors.Emplace(FText::Format(
+			LOCTEXT("Checkout", "Failed to checkout file '{0}' from revision control when persisting sandbox state!"),
+			FText::AsCultureInvariant(InFilename)
+			));
+	}
+	virtual void HandleError_Revert(const TCHAR* InFilename) override
+	{
+		AllErrors.Emplace(FText::Format(
+			LOCTEXT("Revert", "Failed to revert file '{0}' from revision control when persisting sandbox state!"),
+			FText::AsCultureInvariant(InFilename)
+			));
+	}
+	virtual void HandleError_MarkForAdd(const TCHAR* InFilename) override
+	{
+		AllErrors.Emplace(FText::Format(
+			LOCTEXT("MarkForAdd", "Failed to add file '{0}' to revision control when persisting sandbox state!"),
+			FText::AsCultureInvariant(InFilename)
+			));
+	}
+	virtual void HandleError_DeleteSCC(const TCHAR* InFilename) override
+	{
+		AllErrors.Emplace(FText::Format(
+			LOCTEXT("DeleteSCC", "Failed to delete file '{0}' from revision control when persisting sandbox state!"),
+			FText::AsCultureInvariant(InFilename)
+			));
+	}
+	virtual void HandleError_MakeWritable(const TCHAR* InFilename) override
+	{
+		AllErrors.Emplace(FText::Format(
+			LOCTEXT("MakeWritable", "Failed to make '{0}' writeable using file system when persisting sandbox state!"),
+			FText::AsCultureInvariant(InFilename)
+			));
+	}
+	virtual void HandleError_MoveFile(const TCHAR* InToFilename, const TCHAR* InFromFilename) override
+	{
+		AllErrors.Emplace(FText::Format(
+			LOCTEXT("MoveFile", "Failed to move file '{0}' to '{1}' using file system when persisting sandbox state!"),
+			FText::AsCultureInvariant(InFromFilename),
+			FText::AsCultureInvariant(InToFilename)
+			));
+	}
+	virtual void HandleError_DeleteFile(const TCHAR* InFilename) override
+	{
+		AllErrors.Emplace(FText::Format(
+			LOCTEXT("DeleteFile", "Failed to delete file '{0}' using file system when persisting sandbox state!"),
+			FText::AsCultureInvariant(InFilename)
+			));
+	}
+	//~ End IPersistFeedback Interface
+};
+}
+
+#undef LOCTEXT_NAMESPACE
\ No newline at end of file
diff --git a/Engine/Plugins/Developer/Sandbox/FileSandbox/Source/FileSandboxCore/Public/Interface/IPersistFeedback.h b/Engine/Plugins/Developer/Sandbox/FileSandbox/Source/FileSandboxCore/Public/Interface/IPersistFeedback.h
new file mode 100644
--- /dev/null
+++ b/Engine/Plugins/Developer/Sandbox/FileSandbox/Source/FileSandboxCore/Public/Interface/IPersistFeedback.h
@@ -0,0 +1,31 @@
+﻿// Copyright Epic Games, Inc. All Rights Reserved.
+
+#pragma once
+
+#include "Templates/UnrealTemplate.h"
+
+namespace UE::FileSandboxCore
+{
+/** You can implement this to get interactive information about a persist operation. */
+class IPersistFeedback
+{
+public:
+	
+	/** The file was successfully persisted. */
+	virtual void StartFile(const TCHAR* InFilename) = 0;
+	virtual void HandleSuccess(const TCHAR* InFilename) = 0;
+
+	virtual void HandleError_CheckoutNotAllowed(const TCHAR* InFilename) = 0;
+	virtual void HandleError_Checkout(const TCHAR* InFilename) = 0;
+	virtual void HandleError_Revert(const TCHAR* InFilename) = 0;
+	virtual void HandleError_MarkForAdd(const TCHAR* InFilename) = 0;
+	virtual void HandleError_DeleteSCC(const TCHAR* InFilename) = 0;
+	
+	virtual void HandleError_MakeWritable(const TCHAR* InFilename) = 0;
+	virtual void HandleError_MoveFile(const TCHAR* InToFilename, const TCHAR* InFromFilename) = 0;
+	virtual void HandleError_DeleteFile(const TCHAR* InFilename) = 0;
+
+	virtual ~IPersistFeedback() = default;
+};
+}
+
diff --git a/Engine/Plugins/Developer/Sandbox/FileSandbox/Source/FileSandboxCore/Public/Types/Sandbox/PersistArgs.h b/Engine/Plugins/Developer/Sandbox/FileSandbox/Source/FileSandboxCore/Public/Types/Sandbox/PersistArgs.h
--- a/Engine/Plugins/Developer/Sandbox/FileSandbox/Source/FileSandboxCore/Public/Types/Sandbox/PersistArgs.h
+++ b/Engine/Plugins/Developer/Sandbox/FileSandbox/Source/FileSandboxCore/Public/Types/Sandbox/PersistArgs.h
@@ -9,6 +9,8 @@ class FString;
 
 namespace UE::FileSandboxCore
 {
+class IPersistFeedback;
+
 /** Input arguments for persisting a sandbox. */
 struct FPersistArgs
 {
@@ -17,5 +19,14 @@ struct FPersistArgs
 	
 	/** Whether to force the file writable when no source control provider is active. */
 	bool bShouldMakeWritableIfNoSourceControl = false;
+	
+	/** 
+	 * You can implement this to get interactive information about a persist operation.
+	 * Useful for displaying a message about the persist operation or a progress bar. 
+	 * 
+	 * You can aggregate actions.
+	 * @see FAggregatePersistFeedback
+	 */
+	IPersistFeedback* Feedback = nullptr;
 };
 }
diff --git a/Engine/Plugins/Developer/Sandbox/FileSandbox/Source/FileSandboxCore/Public/Types/Sandbox/PersistResult.h b/Engine/Plugins/Developer/Sandbox/FileSandbox/Source/FileSandboxCore/Public/Types/Sandbox/PersistResult.h
--- a/Engine/Plugins/Developer/Sandbox/FileSandbox/Source/FileSandboxCore/Public/Types/Sandbox/PersistResult.h
+++ b/Engine/Plugins/Developer/Sandbox/FileSandbox/Source/FileSandboxCore/Public/Types/Sandbox/PersistResult.h
@@ -24,9 +24,6 @@ enum class EPersistStatus : uint8
 /** Result of a persist.*/
 struct FPersistResult
 {
-	/** List of reasons for persist failure. */
-	TArray<FText> FailureReasons;
-
 	/** Success or failure of the session persist. */
 	EPersistStatus PersistStatus = EPersistStatus::NotAllowed;
 };
