File Sandbox UI: Rename folder Mutex to EntryPoint so it's consistent with the name in the public API

#Synced-CL 49159615

diff --git a/Engine/Plugins/Developer/Sandbox/FileSandbox/Source/FileSandboxUI/Private/EntryPoint/EntryPointWidgetFactory.cpp b/Engine/Plugins/Developer/Sandbox/FileSandbox/Source/FileSandboxUI/Private/EntryPoint/EntryPointWidgetFactory.cpp
new file mode 100644
--- /dev/null
+++ b/Engine/Plugins/Developer/Sandbox/FileSandbox/Source/FileSandboxUI/Private/EntryPoint/EntryPointWidgetFactory.cpp
@@ -0,0 +1,21 @@
+﻿// Copyright Epic Games, Inc. All Rights Reserved.
+
+#include "EntryPoint/EntryPointWidgetFactory.h"
+
+#include "ExternalSandboxActiveViewModel.h"
+#include "Widgets/SExternalSandboxActiveOverlay.h"
+
+namespace UE::FileSandboxUI
+{
+TSharedRef<SWidget> MakeExternalSandboxActiveOverlay(const TSharedRef<IExternalSandboxActiveViewModel>& InViewModel)
+{
+	return SNew(SExternalSandboxActiveOverlay, InViewModel);
+}
+
+TSharedRef<IExternalSandboxActiveViewModel> MakeExternalSandboxActiveViewModel(
+	const TSharedRef<ISandboxEntryPoint>& InOwningEntryPoint
+	)
+{
+	return MakeShared<FExternalSandboxActiveViewModel>(InOwningEntryPoint);
+}
+}
diff --git a/Engine/Plugins/Developer/Sandbox/FileSandbox/Source/FileSandboxUI/Private/EntryPoint/ExternalSandboxActiveViewModel.cpp b/Engine/Plugins/Developer/Sandbox/FileSandbox/Source/FileSandboxUI/Private/EntryPoint/ExternalSandboxActiveViewModel.cpp
new file mode 100644
--- /dev/null
+++ b/Engine/Plugins/Developer/Sandbox/FileSandbox/Source/FileSandboxUI/Private/EntryPoint/ExternalSandboxActiveViewModel.cpp
@@ -0,0 +1,93 @@
+﻿// Copyright Epic Games, Inc. All Rights Reserved.
+
+#include "ExternalSandboxActiveViewModel.h"
+
+#include "EntryPoint/ISandboxEntryPoint.h"
+#include "EntryPoint/ISandboxEntryPointRegistry.h"
+#include "IFileSandboxCoreModule.h"
+#include "IFileSandboxUIModule.h"
+#include "ISandboxManager.h"
+#include "Internationalization/Internationalization.h"
+#include "Internationalization/Text.h"
+
+#define LOCTEXT_NAMESPACE "FExternalSandboxActiveViewModel"
+
+namespace UE::FileSandboxUI
+{
+FExternalSandboxActiveViewModel::FExternalSandboxActiveViewModel(
+	const TSharedRef<ISandboxEntryPoint>& InOwningEntryPoint
+	)
+	: OwningEntryPoint(InOwningEntryPoint)
+{
+	using namespace UE::FileSandboxCore;
+	IFileSandboxCoreModule& ModuleCore = IFileSandboxCoreModule::Get();
+	ISandboxManager& SandboxManager = ModuleCore.GetSandboxManager();
+	SandboxManager.OnPostSandboxStartup().AddRaw(this, &FExternalSandboxActiveViewModel::OnSandboxStartup);
+	SandboxManager.OnPostSandboxShutdown().AddRaw(this, &FExternalSandboxActiveViewModel::OnSandboxShutdown);
+}
+
+FExternalSandboxActiveViewModel::~FExternalSandboxActiveViewModel()
+{
+	using namespace UE::FileSandboxCore;
+	if (IFileSandboxCoreModule::IsAvailable())
+	{
+		IFileSandboxCoreModule& ModuleCore = IFileSandboxCoreModule::Get();
+		ISandboxManager& SandboxManager = ModuleCore.GetSandboxManager();
+		SandboxManager.OnPostSandboxStartup().RemoveAll(this);
+		SandboxManager.OnPostSandboxShutdown().RemoveAll(this);
+	}
+}
+
+bool FExternalSandboxActiveViewModel::IsExternalSandboxActive() const
+{
+	using namespace UE::FileSandboxCore;
+	IFileSandboxCoreModule& ModuleCore = IFileSandboxCoreModule::Get();
+	IFileSandboxUIModule& ModuleUI = IFileSandboxUIModule::Get();
+	return ModuleCore.GetSandboxManager().HasActiveSandbox() 
+		&& ModuleUI.GetEntryPointRegistry().FindOwnerOfActiveSandbox() != OwningEntryPoint;
+}
+
+void FExternalSandboxActiveViewModel::SummonSandboxOwnerUI() const
+{
+	IFileSandboxUIModule& ModuleUI = IFileSandboxUIModule::Get();
+	if (const TSharedPtr<ISandboxEntryPoint> EntryPoint = ModuleUI.GetEntryPointRegistry().FindOwnerOfActiveSandbox())
+	{
+		EntryPoint->SummonProviderUI();
+	}
+}
+
+bool FExternalSandboxActiveViewModel::IsSummoningSupported() const
+{
+	IFileSandboxUIModule& ModuleUI = IFileSandboxUIModule::Get();
+	const TSharedPtr<ISandboxEntryPoint> EntryPoint = ModuleUI.GetEntryPointRegistry().FindOwnerOfActiveSandbox();
+	return EntryPoint && EntryPoint != OwningEntryPoint;
+}
+
+FText FExternalSandboxActiveViewModel::GetSummonActionLabel() const
+{
+	IFileSandboxUIModule& ModuleUI = IFileSandboxUIModule::Get();
+	const TSharedPtr<ISandboxEntryPoint> EntryPoint = ModuleUI.GetEntryPointRegistry().FindOwnerOfActiveSandbox();
+	return EntryPoint ? EntryPoint->GetEntryPointLabel() : FText::GetEmpty();
+}
+
+FText FExternalSandboxActiveViewModel::GetExternalSandboxActiveText() const
+{
+	IFileSandboxUIModule& ModuleUI = IFileSandboxUIModule::Get();
+	
+	if (const TSharedPtr<ISandboxEntryPoint> EntryPoint = ModuleUI.GetEntryPointRegistry().FindOwnerOfActiveSandbox())
+	{
+		return FText::Format(
+			LOCTEXT("Unavailable.KnownEntryPointFmt", "{0} is unavailable because the engine is sandboxed by {1}."),
+			OwningEntryPoint->GetEntryPointLabel(),
+			EntryPoint->GetEntryPointLabel()
+			);
+	}
+	
+	return FText::Format(
+		LOCTEXT("Unavailable.UnknownEntryPointFmt", "{0} is unavailable because the engine is sandboxed."),
+		OwningEntryPoint->GetEntryPointLabel()
+		);
+}
+}
+
+#undef LOCTEXT_NAMESPACE
\ No newline at end of file
diff --git a/Engine/Plugins/Developer/Sandbox/FileSandbox/Source/FileSandboxUI/Private/EntryPoint/ExternalSandboxActiveViewModel.h b/Engine/Plugins/Developer/Sandbox/FileSandbox/Source/FileSandboxUI/Private/EntryPoint/ExternalSandboxActiveViewModel.h
new file mode 100644
--- /dev/null
+++ b/Engine/Plugins/Developer/Sandbox/FileSandbox/Source/FileSandboxUI/Private/EntryPoint/ExternalSandboxActiveViewModel.h
@@ -0,0 +1,41 @@
+﻿// Copyright Epic Games, Inc. All Rights Reserved.
+
+#pragma once
+
+#include "ISandboxInstance.h"
+#include "EntryPoint/IExternalSandboxActiveViewModel.h"
+
+namespace UE::FileSandboxUI
+{
+class ISandboxEntryPoint;
+
+/** This view-model is associated with an ISandboxEntryPoint. The UI is shown when the sandbox is not owned by its entry point. */
+class FExternalSandboxActiveViewModel : public IExternalSandboxActiveViewModel
+{
+public:
+	
+	explicit FExternalSandboxActiveViewModel(const TSharedRef<ISandboxEntryPoint>& InOwningEntryPoint);
+	virtual ~FExternalSandboxActiveViewModel() override;
+	
+	//~ Begin IExternalSandboxActiveViewModel Interface
+	virtual bool IsExternalSandboxActive() const override;
+	virtual void SummonSandboxOwnerUI() const override;
+	virtual bool IsSummoningSupported() const override;
+	virtual FText GetSummonActionLabel() const override;
+	virtual FText GetExternalSandboxActiveText() const override;
+	virtual FSimpleMulticastDelegate& OnVisibilityChanged() override { return OnVisibilityChangedDelegate; }
+	//~ End IExternalSandboxActiveViewModel Interface
+	
+private:
+	
+	/** The entry point for which the UI is being displayed. Used to determine whether an external sandbox is active. */
+	const TSharedRef<ISandboxEntryPoint> OwningEntryPoint;
+	
+	/** Event invoked when the visibility of the widget should change. */
+	FSimpleMulticastDelegate OnVisibilityChangedDelegate;
+	
+	void OnSandboxStartup(FileSandboxCore::ISandboxInstance& SandboxInstance) { return OnVisibilityChangedDelegate.Broadcast(); }
+	void OnSandboxShutdown() { return OnVisibilityChangedDelegate.Broadcast(); }
+};
+}
+
diff --git a/Engine/Plugins/Developer/Sandbox/FileSandbox/Source/FileSandboxUI/Private/EntryPoint/SandboxEntryPointManager.cpp b/Engine/Plugins/Developer/Sandbox/FileSandbox/Source/FileSandboxUI/Private/EntryPoint/SandboxEntryPointManager.cpp
new file mode 100644
--- /dev/null
+++ b/Engine/Plugins/Developer/Sandbox/FileSandbox/Source/FileSandboxUI/Private/EntryPoint/SandboxEntryPointManager.cpp
@@ -0,0 +1,41 @@
+﻿// Copyright Epic Games, Inc. All Rights Reserved.
+
+#include "SandboxEntryPointManager.h"
+
+#include "EntryPoint/ISandboxEntryPoint.h"
+#include "IFileSandboxCoreModule.h"
+#include "ISandboxManager.h"
+
+namespace UE::FileSandboxUI
+{
+void FSandboxEntryPointManager::RegisterEntryPoint(const TSharedRef<ISandboxEntryPoint>& InEntryPoint)
+{
+	EntryPoints.AddUnique(InEntryPoint);
+	OnEntryPointsChangedDelegate.Broadcast();
+}
+
+void FSandboxEntryPointManager::UnregisterEntryPoint(const TSharedRef<ISandboxEntryPoint>& InEntryPoint)
+{
+	EntryPoints.RemoveSingle(InEntryPoint);
+	OnEntryPointsChangedDelegate.Broadcast();
+}
+
+TSharedPtr<ISandboxEntryPoint> FSandboxEntryPointManager::FindOwnerOfActiveSandbox() const
+{
+	using namespace UE::FileSandboxCore;
+	IFileSandboxCoreModule& Module = IFileSandboxCoreModule::Get();
+	const ISandboxInstance* Instance = Module.GetSandboxManager().GetActiveSandboxInstance();
+	if (!Instance)
+	{
+		return nullptr;
+	}
+	
+	// We don't handle multiple EntryPoints reporting ownership.
+	const int32 Index = EntryPoints.IndexOfByPredicate([Instance](const TWeakPtr<ISandboxEntryPoint>& WeakEntryPoint)
+	{
+		const TSharedPtr<ISandboxEntryPoint> EntryPointPin = WeakEntryPoint.Pin();
+		return EntryPointPin && EntryPointPin->OwnsSandbox(*Instance);
+	});
+	return EntryPoints.IsValidIndex(Index) ? EntryPoints[Index].Pin() : nullptr;
+}
+}
diff --git a/Engine/Plugins/Developer/Sandbox/FileSandbox/Source/FileSandboxUI/Private/EntryPoint/SandboxEntryPointManager.h b/Engine/Plugins/Developer/Sandbox/FileSandbox/Source/FileSandboxUI/Private/EntryPoint/SandboxEntryPointManager.h
new file mode 100644
--- /dev/null
+++ b/Engine/Plugins/Developer/Sandbox/FileSandbox/Source/FileSandboxUI/Private/EntryPoint/SandboxEntryPointManager.h
@@ -0,0 +1,30 @@
+﻿// Copyright Epic Games, Inc. All Rights Reserved.
+
+#pragma once
+
+#include "Containers/Array.h"
+#include "EntryPoint/ISandboxEntryPointRegistry.h"
+
+namespace UE::FileSandboxUI
+{
+class FSandboxEntryPointManager : public ISandboxEntryPointRegistry
+{
+public:
+	
+	//~ Begin ISandboxEntryPointRegistry Interface
+	virtual void RegisterEntryPoint(const TSharedRef<ISandboxEntryPoint>& InEntryPoint) override;
+	virtual void UnregisterEntryPoint(const TSharedRef<ISandboxEntryPoint>& InEntryPoint) override;
+	virtual TSharedPtr<ISandboxEntryPoint> FindOwnerOfActiveSandbox() const override;
+	virtual FSimpleMulticastDelegate& OnEntryPointsChanged() override { return OnEntryPointsChangedDelegate; }
+	//~ End ISandboxEntryPointRegistry Interface
+	
+private:
+	
+	/** The entry points this manager knows about. */
+	TArray<TWeakPtr<ISandboxEntryPoint>> EntryPoints;
+	
+	/** Event invoked when the registered entry points have changed. */
+	FSimpleMulticastDelegate OnEntryPointsChangedDelegate;
+};
+}
+
diff --git a/Engine/Plugins/Developer/Sandbox/FileSandbox/Source/FileSandboxUI/Private/EntryPoint/Widgets/SExternalSandboxActiveOverlay.cpp b/Engine/Plugins/Developer/Sandbox/FileSandbox/Source/FileSandboxUI/Private/EntryPoint/Widgets/SExternalSandboxActiveOverlay.cpp
new file mode 100644
--- /dev/null
+++ b/Engine/Plugins/Developer/Sandbox/FileSandbox/Source/FileSandboxUI/Private/EntryPoint/Widgets/SExternalSandboxActiveOverlay.cpp
@@ -0,0 +1,90 @@
+﻿// Copyright Epic Games, Inc. All Rights Reserved.
+
+#include "SExternalSandboxActiveOverlay.h"
+
+#include "Components/VerticalBox.h"
+#include "EntryPoint/IExternalSandboxActiveViewModel.h"
+#include "FileSandboxStyle.h"
+#include "SActionButton.h"
+#include "Widgets/Images/SImage.h"
+#include "Widgets/SOverlay.h"
+#include "Widgets/Text/STextBlock.h"
+
+#define LOCTEXT_NAMESPACE "SExternalSandboxActiveOverlay"
+
+namespace UE::FileSandboxUI
+{
+void SExternalSandboxActiveOverlay::Construct(const FArguments& InArgs, const TSharedRef<IExternalSandboxActiveViewModel>& InViewModel)
+{
+	ViewModel = InViewModel;
+	
+	ChildSlot
+	[
+		SNew(SOverlay)
+
+		+SOverlay::Slot()
+		.ZOrder(0)
+		[
+			SNew(SImage)
+			.Image(FFileSandboxStyle::Get().GetBrush("FileSandbox.TranslucentBackground"))
+		]
+
+		+SOverlay::Slot()
+		.ZOrder(1)
+		.HAlign(HAlign_Center)
+		.VAlign(VAlign_Center)
+		[
+			SNew(SVerticalBox)
+
+			+SVerticalBox::Slot()
+			.AutoHeight()
+			[
+				SNew(STextBlock)
+				.Text(this, &SExternalSandboxActiveOverlay::GetTitleText)
+				.Justification(ETextJustify::Center)
+			]
+
+			+SVerticalBox::Slot()
+			.AutoHeight()
+			.HAlign(HAlign_Center)
+			.Padding(0.f, 10.f)
+			[
+				SNew(SHorizontalBox)
+
+				+SHorizontalBox::Slot()
+				.AutoWidth()
+				[
+					SNew(SActionButton)
+					.Text(this, &SExternalSandboxActiveOverlay::GetSummonButtonLabel)
+					.Visibility_Lambda([this]{ return ViewModel->IsSummoningSupported() ? EVisibility::Visible : EVisibility::Collapsed; })
+					.OnClicked(this, &SExternalSandboxActiveOverlay::OnSummonClicked)
+				]
+			]
+		]
+	];
+	
+	// By setting visiblity here, we avoid duplicate visiblity logic for all API users.
+	SetVisibility(TAttribute<EVisibility>::CreateLambda([this]
+	{
+		return ViewModel->IsExternalSandboxActive() ? EVisibility::Visible : EVisibility::Collapsed;
+	}));
+}
+
+FText SExternalSandboxActiveOverlay::GetTitleText() const
+{
+	return ViewModel->GetExternalSandboxActiveText();
+}
+
+FText SExternalSandboxActiveOverlay::GetSummonButtonLabel() const
+{
+	return FText::Format(LOCTEXT("SummonUI.Label", "Go to {0}"), ViewModel->GetSummonActionLabel());
+}
+
+FReply SExternalSandboxActiveOverlay::OnSummonClicked()
+{
+	ViewModel->SummonSandboxOwnerUI();
+	return FReply::Handled();
+}
+}
+
+#undef LOCTEXT_NAMESPACE
\ No newline at end of file
diff --git a/Engine/Plugins/Developer/Sandbox/FileSandbox/Source/FileSandboxUI/Private/EntryPoint/Widgets/SExternalSandboxActiveOverlay.h b/Engine/Plugins/Developer/Sandbox/FileSandbox/Source/FileSandboxUI/Private/EntryPoint/Widgets/SExternalSandboxActiveOverlay.h
new file mode 100644
--- /dev/null
+++ b/Engine/Plugins/Developer/Sandbox/FileSandbox/Source/FileSandboxUI/Private/EntryPoint/Widgets/SExternalSandboxActiveOverlay.h
@@ -0,0 +1,35 @@
+﻿// Copyright Epic Games, Inc. All Rights Reserved.
+
+#pragma once
+
+#include "Widgets/DeclarativeSyntaxSupport.h"
+#include "Widgets/SCompoundWidget.h"
+
+namespace UE::FileSandboxUI
+{
+class IExternalSandboxActiveViewModel;
+
+/**
+ * Use this overlay in your editor system to tell the user that the engine is sandboxed by another editor system.
+ * This widget shows a button to summon the other system's UI.
+ */
+class SExternalSandboxActiveOverlay : public SCompoundWidget
+{
+public:
+	
+	SLATE_BEGIN_ARGS(SExternalSandboxActiveOverlay) {}
+	SLATE_END_ARGS()
+
+	void Construct(const FArguments& InArgs, const TSharedRef<IExternalSandboxActiveViewModel>& InViewModel);
+	
+private:
+	
+	/** View model that we're displaying. */
+	TSharedPtr<IExternalSandboxActiveViewModel> ViewModel;
+	
+	FText GetTitleText() const;
+	FText GetSummonButtonLabel() const;
+	FReply OnSummonClicked();
+};
+}
+
diff --git a/Engine/Plugins/Developer/Sandbox/FileSandbox/Source/FileSandboxUI/Private/FileSandboxUIModule.cpp b/Engine/Plugins/Developer/Sandbox/FileSandbox/Source/FileSandboxUI/Private/FileSandboxUIModule.cpp
--- a/Engine/Plugins/Developer/Sandbox/FileSandbox/Source/FileSandboxUI/Private/FileSandboxUIModule.cpp
+++ b/Engine/Plugins/Developer/Sandbox/FileSandbox/Source/FileSandboxUI/Private/FileSandboxUIModule.cpp
@@ -2,7 +2,7 @@
 
 #include "FileSandboxUIModule.h"
 
-#include "Mutex/SandboxEntryPointManager.h"
+#include "EntryPoint/SandboxEntryPointManager.h"
 
 namespace UE::FileSandboxUI
 {
diff --git a/Engine/Plugins/Developer/Sandbox/FileSandbox/Source/FileSandboxUI/Private/FileSandboxUIModule.h b/Engine/Plugins/Developer/Sandbox/FileSandbox/Source/FileSandboxUI/Private/FileSandboxUIModule.h
--- a/Engine/Plugins/Developer/Sandbox/FileSandbox/Source/FileSandboxUI/Private/FileSandboxUIModule.h
+++ b/Engine/Plugins/Developer/Sandbox/FileSandbox/Source/FileSandboxUI/Private/FileSandboxUIModule.h
@@ -3,7 +3,7 @@
 #pragma once
 
 #include "IFileSandboxUIModule.h"
-#include "Mutex/SandboxEntryPointManager.h"
+#include "EntryPoint/SandboxEntryPointManager.h"
 
 namespace UE::FileSandboxUI
 {
diff --git a/Engine/Plugins/Developer/Sandbox/FileSandbox/Source/FileSandboxUI/Private/Mutex/EntryPointWidgetFactory.cpp b/Engine/Plugins/Developer/Sandbox/FileSandbox/Source/FileSandboxUI/Private/Mutex/EntryPointWidgetFactory.cpp
deleted file mode 100644
diff --git a/Engine/Plugins/Developer/Sandbox/FileSandbox/Source/FileSandboxUI/Private/Mutex/ExternalSandboxActiveViewModel.cpp b/Engine/Plugins/Developer/Sandbox/FileSandbox/Source/FileSandboxUI/Private/Mutex/ExternalSandboxActiveViewModel.cpp
deleted file mode 100644
diff --git a/Engine/Plugins/Developer/Sandbox/FileSandbox/Source/FileSandboxUI/Private/Mutex/ExternalSandboxActiveViewModel.h b/Engine/Plugins/Developer/Sandbox/FileSandbox/Source/FileSandboxUI/Private/Mutex/ExternalSandboxActiveViewModel.h
deleted file mode 100644
diff --git a/Engine/Plugins/Developer/Sandbox/FileSandbox/Source/FileSandboxUI/Private/Mutex/SandboxEntryPointManager.cpp b/Engine/Plugins/Developer/Sandbox/FileSandbox/Source/FileSandboxUI/Private/Mutex/SandboxEntryPointManager.cpp
deleted file mode 100644
diff --git a/Engine/Plugins/Developer/Sandbox/FileSandbox/Source/FileSandboxUI/Private/Mutex/SandboxEntryPointManager.h b/Engine/Plugins/Developer/Sandbox/FileSandbox/Source/FileSandboxUI/Private/Mutex/SandboxEntryPointManager.h
deleted file mode 100644
diff --git a/Engine/Plugins/Developer/Sandbox/FileSandbox/Source/FileSandboxUI/Private/Mutex/Widgets/SExternalSandboxActiveOverlay.cpp b/Engine/Plugins/Developer/Sandbox/FileSandbox/Source/FileSandboxUI/Private/Mutex/Widgets/SExternalSandboxActiveOverlay.cpp
deleted file mode 100644
diff --git a/Engine/Plugins/Developer/Sandbox/FileSandbox/Source/FileSandboxUI/Private/Mutex/Widgets/SExternalSandboxActiveOverlay.h b/Engine/Plugins/Developer/Sandbox/FileSandbox/Source/FileSandboxUI/Private/Mutex/Widgets/SExternalSandboxActiveOverlay.h
deleted file mode 100644
