Prevent a load module check on shutdown when discarding the sandbox.

#Synced-CL 49551028

diff --git a/Engine/Plugins/Developer/Sandbox/FileSandbox/Source/FileSandboxCore/Private/Sandbox/SandboxInstance.cpp b/Engine/Plugins/Developer/Sandbox/FileSandbox/Source/FileSandboxCore/Private/Sandbox/SandboxInstance.cpp
--- a/Engine/Plugins/Developer/Sandbox/FileSandbox/Source/FileSandboxCore/Private/Sandbox/SandboxInstance.cpp
+++ b/Engine/Plugins/Developer/Sandbox/FileSandbox/Source/FileSandboxCore/Private/Sandbox/SandboxInstance.cpp
@@ -157,8 +157,11 @@ FDiscardResult FSandboxInstance::DiscardSandbox(const FDiscardArgs& InArgs)
 	
 	ManifestContent.Empty();
 	SaveManifest();
-	
-	Private::AppendPackagesWithInMemoryChanges(Result.PackagesPendingHotReload);
+
+	if (!IsEngineExitRequested())
+	{
+		Private::AppendPackagesWithInMemoryChanges(Result.PackagesPendingHotReload);
+	}
 	return Result;
 }
 
