File SandboxCore: Update manifest file whenever a file is added or edited

#Synced-CL 49159615

diff --git a/Engine/Plugins/Developer/Sandbox/FileSandbox/Source/FileSandboxCore/Private/Sandbox/Platform/SandboxPlatformFile.cpp b/Engine/Plugins/Developer/Sandbox/FileSandbox/Source/FileSandboxCore/Private/Sandbox/Platform/SandboxPlatformFile.cpp
--- a/Engine/Plugins/Developer/Sandbox/FileSandbox/Source/FileSandboxCore/Private/Sandbox/Platform/SandboxPlatformFile.cpp
+++ b/Engine/Plugins/Developer/Sandbox/FileSandbox/Source/FileSandboxCore/Private/Sandbox/Platform/SandboxPlatformFile.cpp
@@ -195,7 +195,7 @@ bool FSandboxPlatformFile::DeleteFile(const TCHAR* Filename)
 		}
 
 		// Is file present outside of sandbox? E.g. existed before we entered sandbox. If so, mark as deleted.
-		const bool bIsDeletingNonSandboxFile = !LowerLevel->FileExists(*ResolvedPath.GetSandboxPath());
+		const bool bIsDeletingNonSandboxFile = LowerLevel->FileExists(*ResolvedPath.GetNonSandboxPath());
 		if (bIsDeletingNonSandboxFile)
 		{
 			MarkPathDeleted(ResolvedPath);
@@ -310,7 +310,17 @@ bool FSandboxPlatformFile::MoveFile(const TCHAR* To, const TCHAR* From)
 			// Non-sandbox -> Sandbox
 			if (LowerLevel->MoveFile(*ResolvedToPath.GetSandboxPath(), *ResolvedFromPath.GetNonSandboxPath()))
 			{
-				ClearPathDeleted(ResolvedToPath);
+				const bool bIsEdit = LowerLevel->FileExists(*ResolvedToPath.GetNonSandboxPath());
+				if (IsPathDeleted(ResolvedToPath))
+				{
+					constexpr bool bBroadcast = false;
+					ClearPathDeleted(ResolvedToPath, bBroadcast);
+					BroadcastFileChange(ResolvedToPath, ESandboxFileChange::Removed, bIsEdit ? ESandboxFileChange::Edited : ESandboxFileChange::None);
+				}
+				else
+				{
+					BroadcastFileChange(ResolvedToPath, ESandboxFileChange::None, bIsEdit ? ESandboxFileChange::Edited : ESandboxFileChange::Added);
+				}
 				return true;
 			}
 			return false;
@@ -345,7 +355,7 @@ bool FSandboxPlatformFile::MoveFile(const TCHAR* To, const TCHAR* From)
 			MarkPathDeleted(ResolvedFromPath);
 			NotifyFileDeleted(ResolvedFromPath);
 		}
-
+		
 		return bSuccess;
 	}
 
@@ -935,24 +945,34 @@ void FSandboxPlatformFile::MarkPathDeleted(const FSandboxedPlatformFilePath& InP
 		return;
 	}
 	
-	bool bAlreadyInSet = false;
-	DeletedSandboxPaths.Add(InPath.GetNonSandboxPath(), &bAlreadyInSet);
-	const bool bMadeChanges = !bAlreadyInSet;
+	const bool bMadeChanges = [this, &InPath]
+	{
+		const FScopeLock DeletedSandboxPathsLock(&DeletedSandboxPathsCS);
+		bool bAlreadyInSet = false;
+		DeletedSandboxPaths.Add(InPath.GetNonSandboxPath(), &bAlreadyInSet);
+			
+		const bool bMadeChanges = !bAlreadyInSet;
+		return bMadeChanges;
+	}();
 	
 	if (bMadeChanges)
 	{
-		OnNonSandboxFileDeletedDelegate.Broadcast(InPath, true);
+		BroadcastFileChange(InPath, ESandboxFileChange::None, ESandboxFileChange::Removed);
 	}
 }
 
-void FSandboxPlatformFile::ClearPathDeleted(const FSandboxedPlatformFilePath& InPath)
+void FSandboxPlatformFile::ClearPathDeleted(const FSandboxedPlatformFilePath& InPath, bool bBroadcast)
 {
-	const int32 NumRemoved = DeletedSandboxPaths.Remove(InPath.GetNonSandboxPath());
-	const bool bMadeChanges = NumRemoved > 0;
+	const bool bMadeChanges = [this, &InPath]
+	{
+		const FScopeLock DeletedSandboxPathsLock(&DeletedSandboxPathsCS);
+		const int32 NumRemoved = DeletedSandboxPaths.Remove(InPath.GetNonSandboxPath());
+		return NumRemoved > 0;
+	}();
 	
-	if (bMadeChanges)
+	if (bMadeChanges && bBroadcast)
 	{
-		OnNonSandboxFileDeletedDelegate.Broadcast(InPath, false);
+		BroadcastFileChange(InPath, ESandboxFileChange::Removed, ESandboxFileChange::None);
 	}
 }
 
@@ -1104,4 +1124,12 @@ void FSandboxPlatformFile::OnSandboxDirectoryChanged(const TArray<FFileChangeDat
 	}
 }
 #endif
+
+
+void FSandboxPlatformFile::BroadcastFileChange(const FSandboxedPlatformFilePath& InPath, ESandboxFileChange InOldState, ESandboxFileChange InNewState)
+{
+	// This would be pointless. Investigate.
+	ensure(InOldState != InNewState);
+	OnFileStateChangedDelegate.Broadcast(InPath, InOldState, InNewState);
+}
 }
\ No newline at end of file
diff --git a/Engine/Plugins/Developer/Sandbox/FileSandbox/Source/FileSandboxCore/Private/Sandbox/Platform/SandboxPlatformFile.h b/Engine/Plugins/Developer/Sandbox/FileSandbox/Source/FileSandboxCore/Private/Sandbox/Platform/SandboxPlatformFile.h
--- a/Engine/Plugins/Developer/Sandbox/FileSandbox/Source/FileSandboxCore/Private/Sandbox/Platform/SandboxPlatformFile.h
+++ b/Engine/Plugins/Developer/Sandbox/FileSandbox/Source/FileSandboxCore/Private/Sandbox/Platform/SandboxPlatformFile.h
@@ -27,6 +27,9 @@ struct FPersistResult;
 struct FSandboxMountPoint;
 	
 DECLARE_MULTICAST_DELEGATE_TwoParams(FDeletionStateChanged, const FSandboxedPlatformFilePath&, bool /* bIsDeleted */);
+DECLARE_MULTICAST_DELEGATE_OneParam(FSandboxedFilePathDelegate, const FSandboxedPlatformFilePath&);
+
+DECLARE_MULTICAST_DELEGATE_ThreeParams(FFileStateChangeDelegate, const FSandboxedPlatformFilePath&, ESandboxFileChange /*OldState*/, ESandboxFileChange /*NewState*/);
 	
 class FSandboxPlatformFile : public IPlatformFile
 {
@@ -101,11 +104,8 @@ public:
 	/** @return Whether the given package file was deleted and exists on non sandbox path. */
 	bool DeletedPackageExistsInNonSandbox(FString InFilename) const;
 
-	/**
-	 * Invoked when a file from a non-sandbox path is deleted, i.e. a file that existed before the sandbox is deleted.
-	 * Specifically: NOT invoked for files that were added during sandbox, and then deleted again.
-	 */
-	FDeletionStateChanged& OnNonSandboxFileDeleted() { return OnNonSandboxFileDeletedDelegate; }
+	/** Invoked when the file action of a file changes. */
+	FFileStateChangeDelegate& OnFileStateChanged() { return OnFileStateChangedDelegate; }
 	
 private:
 	
@@ -139,7 +139,7 @@ private:
 	/** Set that the non-sandbox file or directory is explicitly deleted by the sandbox. */
 	void MarkPathDeleted(const FSandboxedPlatformFilePath& InPath);
 	/** Clears the deletion state. The file or directory is no longer explicitly deleted by the sandbox. */
-	void ClearPathDeleted(const FSandboxedPlatformFilePath& InPath);
+	void ClearPathDeleted(const FSandboxedPlatformFilePath& InPath, bool bBroadcast = true);
 
 	/** Notify that a file has been explicitly deleted from the sandbox */
 	void NotifyFileDeleted(const FSandboxedPlatformFilePath& InPath);
@@ -154,6 +154,9 @@ private:
 	/** Called when a file in a sandbox directory changes on disk */
 	void OnSandboxDirectoryChanged(const TArray<FFileChangeData>& FileChanges, FSandboxedPlatformFilePath MountPath);
 #endif
+	
+	/** Broadcasts OnFileStateChangedDelegate and checks invariants. */
+	void BroadcastFileChange(const FSandboxedPlatformFilePath& InPath, ESandboxFileChange InOldState, ESandboxFileChange InNewState);
 
 	/** Attempts to persist the file. */
 	ESandboxFileChange PersistItem(
@@ -188,10 +191,15 @@ private:
 	mutable FCriticalSection DeletedSandboxPathsCS;
 	
 	/**
-	 * Invoked when a file from a non-sandbox path is deleted, i.e. a file that existed before the sandbox is deleted.
-	 * Specifically: NOT invoked for files that were added during sandbox, and then deleted again.
+	 * Invoked when a file action changes.
+	 * 
+	 * Examples:
+	 * - Adding Foo.uasset: (Foo.uasset, None, Added)
+	 * - Reverting the addition of Foo.uasset: (Foo.uasset, Added, None)
+	 * 
+	 * There should never be any permutation of old = None, new = None.
 	 */
-	FDeletionStateChanged OnNonSandboxFileDeletedDelegate;
+	FFileStateChangeDelegate OnFileStateChangedDelegate;
 };
 }
 
diff --git a/Engine/Plugins/Developer/Sandbox/FileSandbox/Source/FileSandboxCore/Private/Sandbox/SandboxInstance.cpp b/Engine/Plugins/Developer/Sandbox/FileSandbox/Source/FileSandboxCore/Private/Sandbox/SandboxInstance.cpp
--- a/Engine/Plugins/Developer/Sandbox/FileSandbox/Source/FileSandboxCore/Private/Sandbox/SandboxInstance.cpp
+++ b/Engine/Plugins/Developer/Sandbox/FileSandbox/Source/FileSandboxCore/Private/Sandbox/SandboxInstance.cpp
@@ -3,6 +3,7 @@
 #include "SandboxInstance.h"
 
 #include "Interface/Feedback/AggregatePersistFeedback.h"
+#include "Misc/CoreDelegates.h"
 #include "Misc/PackageName.h"
 #include "Persist/DiffUtils.h"
 #include "Persist/LogPersistFeedback.h"
@@ -31,6 +32,18 @@ static void TrackRevertedFileChangeInManifest(FFileSandboxCore_ManifestData& Out
 	}
 }
 
+static void TrackFileChangeInManifest(FFileSandboxCore_ManifestData& OutManifestData, const FString& InFile, ESandboxFileChange InAction)
+{
+	switch (InAction)
+	{
+	case ESandboxFileChange::Added: OutManifestData.AddedFiles.Add(InFile); break;
+	case ESandboxFileChange::Edited: OutManifestData.ModifiedFiles.Add(InFile); break;
+	case ESandboxFileChange::Removed: OutManifestData.DeletedFiles.Add(InFile); break;
+	default: break;
+	}
+}
+
+
 /** 
  * Adds packages with in-memory changes. 
  * 
@@ -78,20 +91,22 @@ FSandboxInstance::FSandboxInstance(
 {
 	PlatformFile->Initialize(&InitArgs.GetPlatformFile(), TEXT(""));
 	ReplayChanges(ManifestContent.DeletedFiles.GetFiles(), ManifestContent.ModifiedFiles.GetFiles(), ManifestContent.AddedFiles.GetFiles());
-	PlatformFile->OnNonSandboxFileDeleted().AddRaw(this, &FSandboxInstance::OnFileDeletionStateChanged);
+	
+	PlatformFile->OnFileStateChanged().AddRaw(this, &FSandboxInstance::OnFileDeletionStateChanged);
+	
+	FCoreDelegates::OnEndFrame.AddRaw(this, &FSandboxInstance::OnEndOfFrame);
 }
 	
 FSandboxInstance::~FSandboxInstance()
 {
-	// TODO DP UE-353328: In the future, we should save these changes out as the happen. Otherwise, we cannot recover from crashes that happen during sandbox.
+	FCoreDelegates::OnEndFrame.RemoveAll(this);
+	
 	ManifestContent.Empty();
-	const auto ProcessRevertedChange = [this](const FFileChange& Change)
+	TArray<FName> PackagesPendingHotReload, PackagesPendingPurge;
+	PlatformFile->DiscardSandbox(PackagesPendingHotReload, PackagesPendingPurge, [this](const FFileChange& Change)
 	{
 		Private::TrackRevertedFileChangeInManifest(ManifestContent, Change); 
-	};
-	
-	TArray<FName> PackagesPendingHotReload, PackagesPendingPurge;
-	PlatformFile->DiscardSandbox(PackagesPendingHotReload, PackagesPendingPurge, ProcessRevertedChange);
+	});
 	PlatformFile.Reset();
 
 	// Save the data before hot reloading or purging: in case that causes a crash we'll have an update to manifest file. 
@@ -142,6 +157,9 @@ FPersistResult FSandboxInstance::PersistSandbox(const FPersistArgs& InArgs)
 	if (bManifestNeedsSave)
 	{
 		SaveManifest();
+		
+		// Persisting has removed some file states, so inform anyone that cares about it.
+		OnSandboxedFilesChangedDelegate.Broadcast();
 	}
 
 	return FPersistResult
@@ -157,6 +175,9 @@ FDiscardResult FSandboxInstance::DiscardSandbox(const FDiscardArgs& InArgs)
 	
 	ManifestContent.Empty();
 	SaveManifest();
+	
+	// Discarding has removed all file states, so inform anyone that cares about it.
+	OnSandboxedFilesChangedDelegate.Broadcast();
 
 	if (!IsEngineExitRequested())
 	{
@@ -175,18 +196,34 @@ bool FSandboxInstance::DeletedPackageExistsInNonSandbox(const FString& InFilenam
 	return PlatformFile->DeletedPackageExistsInNonSandbox(InFilename);
 }
 
-void FSandboxInstance::OnFileDeletionStateChanged(const FSandboxedPlatformFilePath& InFile, bool bIsDeleted)
+void FSandboxInstance::OnFileDeletionStateChanged(const FSandboxedPlatformFilePath& InFile, ESandboxFileChange InOldState, ESandboxFileChange InNewState)
 {
 	const FNonSandboxPath NonSandboxPath = InFile.GetNonSandboxPath();
 	
-	bool bNeedsSave = false;
-	bNeedsSave |= bIsDeleted && ManifestContent.ModifiedFiles.Remove(NonSandboxPath);
-	bNeedsSave |= bIsDeleted && ManifestContent.DeletedFiles.Add(NonSandboxPath);
-	bNeedsSave |= !bIsDeleted && ManifestContent.DeletedFiles.Remove(NonSandboxPath);
-	
-	if (bNeedsSave)
+	const bool bWasDeleted = InNewState == ESandboxFileChange::Removed;
+	if (bWasDeleted)
 	{
+		ManifestContent.ModifiedFiles.Remove(NonSandboxPath);
+		ManifestContent.DeletedFiles.Add(NonSandboxPath);
+		ManifestContent.AddedFiles.Remove(NonSandboxPath);
 		SaveManifest();
+		
+		// The manifest was updated already. Only need to broadcast at end of tick.
+		const ESandboxRefreshFlags CurrentFlags = RefreshFlags.load();
+		RefreshFlags = CurrentFlags | ESandboxRefreshFlags::BroadcastChanges;
+	}
+	
+	const bool bClearedDeleted = InOldState == ESandboxFileChange::Removed;
+	if (bClearedDeleted)
+	{
+		ManifestContent.DeletedFiles.Remove(NonSandboxPath);
+		// At the end of tick, check whether the file was marked for add or edit
+		RefreshFlags = ESandboxRefreshFlags::BroadcastChanges | ESandboxRefreshFlags::UpdateManifest;
+	}
+	else
+	{
+		// At end of tick, rebuild the manifest based on the files in the sandbox directory.
+		RefreshFlags = ESandboxRefreshFlags::BroadcastChanges | ESandboxRefreshFlags::UpdateManifest;
 	}
 }
 
@@ -194,4 +231,35 @@ bool FSandboxInstance::SaveManifest() const
 {
 	return FileSandboxCore::SaveManifest(ManifestContent, GetRootDirectory());
 }
+
+void FSandboxInstance::OnEndOfFrame()
+{
+	const ESandboxRefreshFlags CurrentFlags = RefreshFlags.exchange(ESandboxRefreshFlags::None);
+	if (CurrentFlags == ESandboxRefreshFlags::None)
+	{
+		return;
+	}
+	
+	// We update the manifest every time an update happens.
+	// This is important as during the sandbox a crash could happen.
+	// This way, when the sandbox is loaded again after a crash it's up to date. 
+	if (EnumHasAnyFlags(CurrentFlags, ESandboxRefreshFlags::UpdateManifest))
+	{
+		FFileSandboxCore_ManifestData NewManifest;
+		NewManifest.VersionInfo.Initialize();
+		EnumerateFileChanges([&NewManifest](const FString& InPath, ESandboxFileChange InChange)
+		{
+			Private::TrackFileChangeInManifest(NewManifest, InPath, InChange);
+			return EBreakBehavior::Continue;
+		});
+		
+		ManifestContent = MoveTemp(NewManifest);
+		SaveManifest();
+	}
+	
+	if (EnumHasAnyFlags(CurrentFlags, ESandboxRefreshFlags::BroadcastChanges))
+	{
+		OnSandboxedFilesChangedDelegate.Broadcast();
+	}
+}
 }
diff --git a/Engine/Plugins/Developer/Sandbox/FileSandbox/Source/FileSandboxCore/Private/Sandbox/SandboxInstance.h b/Engine/Plugins/Developer/Sandbox/FileSandbox/Source/FileSandboxCore/Private/Sandbox/SandboxInstance.h
--- a/Engine/Plugins/Developer/Sandbox/FileSandbox/Source/FileSandboxCore/Private/Sandbox/SandboxInstance.h
+++ b/Engine/Plugins/Developer/Sandbox/FileSandbox/Source/FileSandboxCore/Private/Sandbox/SandboxInstance.h
@@ -3,6 +3,7 @@
 #pragma once
 
 #include "ISandboxInstance.h"
+#include "SandboxRefreshFlags.h"
 #include "Data/ManifestData.h"
 #include "Data/SandboxMetaData.h"
 #include "Platform/SandboxPlatformFile.h"
@@ -11,6 +12,7 @@
 
 namespace UE::FileSandboxCore
 {
+enum class ESandboxRefreshFlags : uint8;
 struct FSandboxInitArgs;
 	
 class FSandboxInstance : public ISandboxInstance, public FNoncopyable
@@ -41,6 +43,7 @@ public:
 	virtual bool DeletedPackageExistsInNonSandbox(const FString& InFilename) const override;
 	virtual const FFileSandboxCore_SandboxMetaData& GetInitialMetaData() const override { return MetaData; }
 	virtual const TCHAR* GetRootDirectory() const override { return *RootDirectory; }
+	virtual FSimpleMulticastDelegate& OnSandboxedFilesChanged() override { return OnSandboxedFilesChangedDelegate; }
 	//~ End ISandboxInstance Interface
 
 private:
@@ -63,12 +66,21 @@ private:
 	
 	/** Sandboxes source control operations and reroutes them to use instead. */
 	FSandboxedSourceControl SourceControlSandbox;
+	
+	/** Delegate invoked when the result of GatherChangedFiles has changed. */
+	FSimpleMulticastDelegate OnSandboxedFilesChangedDelegate;
+	
+	/** Whether any sandboxed file changes have occured this frame. */
+	std::atomic<ESandboxRefreshFlags> RefreshFlags = ESandboxRefreshFlags::None;
 
 	/** Marks the file as deleted in the manifest. */
-	void OnFileDeletionStateChanged(const FSandboxedPlatformFilePath& InFile, bool bIsDeleted);
+	void OnFileDeletionStateChanged(const FSandboxedPlatformFilePath& InFile, ESandboxFileChange InOldState, ESandboxFileChange InNewState);
 
 	/** Saves the manifest file. */
 	bool SaveManifest() const;
+	
+	/** Processes pending work. */
+	void OnEndOfFrame();
 };
 }
 
diff --git a/Engine/Plugins/Developer/Sandbox/FileSandbox/Source/FileSandboxCore/Private/Sandbox/SandboxRefreshFlags.h b/Engine/Plugins/Developer/Sandbox/FileSandbox/Source/FileSandboxCore/Private/Sandbox/SandboxRefreshFlags.h
new file mode 100644
--- /dev/null
+++ b/Engine/Plugins/Developer/Sandbox/FileSandbox/Source/FileSandboxCore/Private/Sandbox/SandboxRefreshFlags.h
@@ -0,0 +1,21 @@
+﻿// Copyright Epic Games, Inc. All Rights Reserved.
+
+#pragma once
+
+#include "HAL/Platform.h"
+#include "Misc/EnumClassFlags.h"
+
+namespace UE::FileSandboxCore
+{
+/** Flags for refreshing the sandbox instance. */
+enum class ESandboxRefreshFlags : uint8
+{
+	None,
+	
+	/** Broadcast that the set of sandboxed file changes has changed. */
+	BroadcastChanges = 1 << 0,
+	/** The manifest file needs to be resaved. */
+	UpdateManifest = 1 << 1
+};
+ENUM_CLASS_FLAGS(ESandboxRefreshFlags);
+}
\ No newline at end of file
diff --git a/Engine/Plugins/Developer/Sandbox/FileSandbox/Source/FileSandboxCore/Public/ISandboxInstance.h b/Engine/Plugins/Developer/Sandbox/FileSandbox/Source/FileSandboxCore/Public/ISandboxInstance.h
--- a/Engine/Plugins/Developer/Sandbox/FileSandbox/Source/FileSandboxCore/Public/ISandboxInstance.h
+++ b/Engine/Plugins/Developer/Sandbox/FileSandbox/Source/FileSandboxCore/Public/ISandboxInstance.h
@@ -87,6 +87,9 @@ public:
 
 	/** @return The root directory of the sandbox. */
 	virtual const TCHAR* GetRootDirectory() const = 0;
+	
+	/** @return Delegate invoked when the result of GatherChangedFiles has changed. */
+	virtual FSimpleMulticastDelegate& OnSandboxedFilesChanged() = 0;
 
 	virtual ~ISandboxInstance() = default;
 };
