File Sandbox Core: Renaming files while connected to source control is now sandboxed by the sandbox system.

#Synced-CL 49159615

diff --git a/Engine/Plugins/Developer/Sandbox/FileSandbox/Source/FileSandboxCore/Private/SourceControl/Proxy/SandboxSourceControlProviderProxy.cpp b/Engine/Plugins/Developer/Sandbox/FileSandbox/Source/FileSandboxCore/Private/SourceControl/Proxy/SandboxSourceControlProviderProxy.cpp
--- a/Engine/Plugins/Developer/Sandbox/FileSandbox/Source/FileSandboxCore/Private/SourceControl/Proxy/SandboxSourceControlProviderProxy.cpp
+++ b/Engine/Plugins/Developer/Sandbox/FileSandbox/Source/FileSandboxCore/Private/SourceControl/Proxy/SandboxSourceControlProviderProxy.cpp
@@ -6,6 +6,7 @@
 #include "GenericPlatform/GenericPlatformFile.h"
 #include "HAL/PlatformFileManager.h"
 #include "SandboxSourceControlStateProxy.h"
+#include "SourceControlOperations.h"
 
 namespace UE::FileSandboxCore
 {
@@ -50,6 +51,18 @@ static TOptional<ECommandResult::Type> HandleSandboxedOp(const FSourceControlOpe
 		return bAllSuccessful ? ECommandResult::Succeeded : ECommandResult::Failed;
 	}
 	
+	if (OpName == TEXT("Copy"))
+	{
+		if (ensure(InFiles.Num() == 1))
+		{
+			TSharedRef<FCopy> CopyOp = StaticCastSharedRef<FCopy>(InOperation);
+			IPlatformFile& PlatformFile = FPlatformFileManager::Get().GetPlatformFile();
+			const bool bSuccess = PlatformFile.MoveFile(*CopyOp->GetDestination(), *InFiles[0]);
+			return bSuccess ? ECommandResult::Succeeded : ECommandResult::Failed;
+		}
+		return ECommandResult::Failed;
+	}
+	
 	return {};
 }
 	
