File Sandbox Core: Add API for getting timestamp of sandboxed change.
- Edits and adds are determined using the file system timestamp
- Removals are determined by saving the time of removal in the the manifest

#Synced-CL 49159615

diff --git a/Engine/Plugins/Developer/Sandbox/FileSandbox/Source/FileSandboxCore/Private/Data/ManifestData.h b/Engine/Plugins/Developer/Sandbox/FileSandbox/Source/FileSandboxCore/Private/Data/ManifestData.h
--- a/Engine/Plugins/Developer/Sandbox/FileSandbox/Source/FileSandboxCore/Private/Data/ManifestData.h
+++ b/Engine/Plugins/Developer/Sandbox/FileSandbox/Source/FileSandboxCore/Private/Data/ManifestData.h
@@ -3,6 +3,7 @@
 #pragma once
 
 #include "AbsoluteFilePathArray.h"
+#include "TimedAbsoluteFilePathArray.h"
 #include "Data/VersionInfo.h"
 #include "ManifestData.generated.h"
 
@@ -15,9 +16,12 @@ struct FFileSandboxCore_ManifestData
 {
 	GENERATED_BODY()
 	
-	/** Absolute non-sandbox paths to files deleted in sandbox. */
+	/**
+	 * Absolute non-sandbox paths to files deleted in sandbox. 
+	 * Also saved the timestamp of the remove operation.
+	 */
 	UPROPERTY()
-	FFileSandboxCore_AbsoluteFilePathArray DeletedFiles;
+	FFileSandboxCore_TimedAbsoluteFilePathArray DeletedFiles;
 
 	/** Absolute non-sandbox paths to files modified in sandbox. */
 	UPROPERTY()
diff --git a/Engine/Plugins/Developer/Sandbox/FileSandbox/Source/FileSandboxCore/Private/Data/TimedAbsoluteFilePathArray.h b/Engine/Plugins/Developer/Sandbox/FileSandbox/Source/FileSandboxCore/Private/Data/TimedAbsoluteFilePathArray.h
new file mode 100644
--- /dev/null
+++ b/Engine/Plugins/Developer/Sandbox/FileSandbox/Source/FileSandboxCore/Private/Data/TimedAbsoluteFilePathArray.h
@@ -0,0 +1,75 @@
+﻿// Copyright Epic Games, Inc. All Rights Reserved.
+
+#pragma once
+
+#include "Misc/Paths.h"
+#include "Misc/DateTime.h"
+#include "TimedAbsoluteFilePathArray.generated.h"
+
+/** Keeps an array of absolute file paths and their respective timestamps. Each path is unique. */
+USTRUCT()
+struct FFileSandboxCore_TimedAbsoluteFilePathArray
+{
+	GENERATED_BODY()
+
+	const TArray<FString>& GetFiles() const { return Files; }
+	
+	const TArray<FDateTime>& GetTimestamps() const { return Timestamps; }
+	bool AreTimestampsValid() const { return Files.Num() == Timestamps.Num(); }
+
+	/** @return Whether the path was added, i.e. was not already present. */
+	bool Add(const FString& InPath, const FDateTime& InTimestamp = FDateTime::Now())
+	{
+		const int32 NumBefore = Files.Num();
+		const FString Path = FPaths::ConvertRelativePathToFull(InPath);
+		
+		if (!Files.Contains(Path))
+		{
+			Files.Add(Path);
+			Timestamps.Add(InTimestamp);
+		}
+		
+		return NumBefore != Files.Num();
+	}
+
+	/** @return Whether anything was removed. */
+	bool Remove(const FString& InPath)
+	{
+		const int32 Index = Files.IndexOfByKey(FPaths::ConvertRelativePathToFull(InPath));
+		
+		const bool bRemove = Files.IsValidIndex(Index);
+		if (bRemove)
+		{
+			Files.RemoveAt(Index);
+		}
+		
+		// User may have edited the manifest file and corrupted it
+		if (bRemove && ensure(Timestamps.IsValidIndex(Index)))
+		{
+			Timestamps.RemoveAt(Index);
+		}
+		
+		return bRemove;
+	}
+
+	void Empty()
+	{
+		Files.Empty();
+		Timestamps.Empty();
+	}
+
+private:
+
+	/** Absolute file paths. */
+	UPROPERTY()
+	TArray<FString> Files;
+	
+	/**
+	 * The times that the files were changed in local time. 
+	 * 
+	 * Under normal circumstances, this is equal length to Files.
+	 * Keep in mind that the user may have edited the file in the manifest, in which case the values in this array are useless.
+	 */
+	UPROPERTY()
+	TArray<FDateTime> Timestamps;
+};
\ No newline at end of file
diff --git a/Engine/Plugins/Developer/Sandbox/FileSandbox/Source/FileSandboxCore/Private/ISandboxManager.cpp b/Engine/Plugins/Developer/Sandbox/FileSandbox/Source/FileSandboxCore/Private/ISandboxManager.cpp
--- a/Engine/Plugins/Developer/Sandbox/FileSandbox/Source/FileSandboxCore/Private/ISandboxManager.cpp
+++ b/Engine/Plugins/Developer/Sandbox/FileSandbox/Source/FileSandboxCore/Private/ISandboxManager.cpp
@@ -3,6 +3,8 @@
 #include "ISandboxManager.h"
 
 #include "Interface/ISandboxLock.h"
+#include "Misc/DateTime.h"
+#include "Types/EBreakBehavior.h"
 #include "Types/Manager/DeleteSandboxByDirectoryArgs.h"
 #include "Types/Manager/DeleteSandboxByNameArgs.h"
 #include "Types/Manager/DeleteSandboxResult.h"
diff --git a/Engine/Plugins/Developer/Sandbox/FileSandbox/Source/FileSandboxCore/Private/Sandbox/ISandboxInstance.cpp b/Engine/Plugins/Developer/Sandbox/FileSandbox/Source/FileSandboxCore/Private/Sandbox/ISandboxInstance.cpp
--- a/Engine/Plugins/Developer/Sandbox/FileSandbox/Source/FileSandboxCore/Private/Sandbox/ISandboxInstance.cpp
+++ b/Engine/Plugins/Developer/Sandbox/FileSandbox/Source/FileSandboxCore/Private/Sandbox/ISandboxInstance.cpp
@@ -2,6 +2,7 @@
 
 #include "ISandboxInstance.h"
 
+#include "Misc/DateTime.h"
 #include "Types/EBreakBehavior.h"
 #include "Types/Sandbox/DiscardArgs.h"
 #include "Types/Sandbox/DiscardResult.h"
@@ -27,13 +28,18 @@ void ISandboxInstance::DiscardAll()
 FGatheredFileChanges ISandboxInstance::GatherChangedFiles(EFileChangeGatherFlags InFlags) const
 {
 	FGatheredFileChanges Result;
-	EnumerateFileChanges([InFlags, &Result](const FString& InFile, ESandboxFileChange InChange)
+	EnumerateFileChanges([this, InFlags, &Result](const FString& InFile, ESandboxFileChange InChange)
 	{
 		Result.NonSandboxPaths.Add(InFile);
 		if (EnumHasAnyFlags(InFlags, EFileChangeGatherFlags::IncludeChangeTypes))
 		{
 			Result.FileActions.Add(InChange);
 		}
+		if (EnumHasAnyFlags(InFlags, EFileChangeGatherFlags::IncludeTimestamps))
+		{
+			const TOptional<FDateTime> Timestamp = GetSandboxedFileTimestamp(InFile);
+			Result.Timestamps.Add(Timestamp.Get(FDateTime::MinValue()));
+		}
 		return EBreakBehavior::Continue;
 	});
 	return Result;
diff --git a/Engine/Plugins/Developer/Sandbox/FileSandbox/Source/FileSandboxCore/Private/Sandbox/Persist/DiffUtils.cpp b/Engine/Plugins/Developer/Sandbox/FileSandbox/Source/FileSandboxCore/Private/Sandbox/Persist/DiffUtils.cpp
--- a/Engine/Plugins/Developer/Sandbox/FileSandbox/Source/FileSandboxCore/Private/Sandbox/Persist/DiffUtils.cpp
+++ b/Engine/Plugins/Developer/Sandbox/FileSandbox/Source/FileSandboxCore/Private/Sandbox/Persist/DiffUtils.cpp
@@ -6,8 +6,10 @@
 #include "Containers/Set.h"
 #include "Data/ManifestData.h"
 #include "GenericPlatform/GenericPlatformFile.h"
+#include "HAL/FileManager.h"
 #include "Misc/PackageName.h"
 #include "Misc/Paths.h"
+#include "Sandbox/Platform/SanboxedFilePathUtils.h"
 #include "Sandbox/Platform/SandboxMountPoint.h"
 #include "Templates/Function.h"
 #include "Types/EBreakBehavior.h"
@@ -100,20 +102,6 @@ void ScanForEditedOrAddedFiles(
 	}
 }
 
-void EnumerateMountPoints(TFunctionRef<EBreakBehavior(const FString& InAssetPath, const FString& InFilesystemPath)> InProcess)
-{
-	TArray<FString> RootPaths;
-	FPackageName::QueryRootContentPaths(RootPaths);
-	for (const FString& AssetPath : RootPaths)
-	{
-		const FString FilesystemPath = FPackageName::LongPackageNameToFilename(AssetPath);
-		if (InProcess(AssetPath, FilesystemPath) == EBreakBehavior::Break)
-		{
-			break;
-		}
-	}	
-}
-
 void EnumerateFileChanges(
 	IPlatformFile& InPlatformFile,
 	const FString& InSandboxRootPath, 
diff --git a/Engine/Plugins/Developer/Sandbox/FileSandbox/Source/FileSandboxCore/Private/Sandbox/Persist/DiffUtils.h b/Engine/Plugins/Developer/Sandbox/FileSandbox/Source/FileSandboxCore/Private/Sandbox/Persist/DiffUtils.h
--- a/Engine/Plugins/Developer/Sandbox/FileSandbox/Source/FileSandboxCore/Private/Sandbox/Persist/DiffUtils.h
+++ b/Engine/Plugins/Developer/Sandbox/FileSandbox/Source/FileSandboxCore/Private/Sandbox/Persist/DiffUtils.h
@@ -4,6 +4,7 @@
 
 #include "ISandboxManager.h"
 #include "Containers/ContainersFwd.h"
+#include "Data/ManifestData.h"
 #include "Sandbox/Platform/SandboxedPlatformFilePath.h"
 #include "Templates/FunctionFwd.h"
 #include "Types/EBreakBehavior.h"
@@ -64,9 +65,6 @@ void ScanForEditedOrAddedFiles(
 	TFunctionRef<FHandleFilePathSignature> InHandleAddedFile
 	);
 
-/** Enumerates all mount points in the engine. */
-void EnumerateMountPoints(TFunctionRef<EBreakBehavior(const FString& InAssetPath, const FString& InFilesystemPath)> InProcess);
-
 /** 
  * Shared logic for determining the effective file changes.
  * @see ISandboxInstance::EnumerateFileChanges.
diff --git a/Engine/Plugins/Developer/Sandbox/FileSandbox/Source/FileSandboxCore/Private/Sandbox/Platform/SanboxedFilePathUtils.cpp b/Engine/Plugins/Developer/Sandbox/FileSandbox/Source/FileSandboxCore/Private/Sandbox/Platform/SanboxedFilePathUtils.cpp
new file mode 100644
--- /dev/null
+++ b/Engine/Plugins/Developer/Sandbox/FileSandbox/Source/FileSandboxCore/Private/Sandbox/Platform/SanboxedFilePathUtils.cpp
@@ -0,0 +1,30 @@
+﻿// Copyright Epic Games, Inc. All Rights Reserved.
+
+#include "SanboxedFilePathUtils.h"
+
+#include "SandboxedPlatformFilePath.h"
+#include "Misc/Optional.h"
+
+namespace UE::FileSandboxCore
+{
+bool IsFileInMountPoint(const FString& InNonSandboxFilePath, const FSandboxedPlatformFilePath& InMountPointRoot)
+{
+	// Mount point are stored with a trailing slash to prevent matching mount point with similar names -> (/Bla/Content, /Bla/ContentSupreme)
+	// So we test without the slash to make sure with can match mount point directly -> (/Bla/Content match /Bla/Content/)
+	const FString& PathStr = InMountPointRoot.GetNonSandboxPath();
+	int32 PathStrNoSlashLength = PathStr.Len() - 1;
+	return FCString::Strnicmp(*InNonSandboxFilePath, *PathStr, PathStrNoSlashLength) == 0 
+		&& (InNonSandboxFilePath.Len() == PathStrNoSlashLength 
+			|| InNonSandboxFilePath[PathStrNoSlashLength] == TEXT('/')
+			);
+}
+
+TOptional<FSandboxedPlatformFilePath> MakeSandboxedFilePathIfInMountPoint(
+	const FString& InNonSandboxFilePath, const FSandboxedPlatformFilePath& InMountPointRoot
+	)
+{
+	return IsFileInMountPoint(InNonSandboxFilePath, InMountPointRoot)
+		? FSandboxedPlatformFilePath::CreateSandboxPath(FString(InNonSandboxFilePath), InMountPointRoot)
+		: TOptional<FSandboxedPlatformFilePath>();
+}
+}
diff --git a/Engine/Plugins/Developer/Sandbox/FileSandbox/Source/FileSandboxCore/Private/Sandbox/Platform/SanboxedFilePathUtils.h b/Engine/Plugins/Developer/Sandbox/FileSandbox/Source/FileSandboxCore/Private/Sandbox/Platform/SanboxedFilePathUtils.h
new file mode 100644
--- /dev/null
+++ b/Engine/Plugins/Developer/Sandbox/FileSandbox/Source/FileSandboxCore/Private/Sandbox/Platform/SanboxedFilePathUtils.h
@@ -0,0 +1,21 @@
+﻿// Copyright Epic Games, Inc. All Rights Reserved.
+
+#pragma once
+
+template<typename OptionalType> struct TOptional;
+
+class FString;
+
+namespace UE::FileSandboxCore
+{
+class FSandboxedPlatformFilePath;
+
+/** @return Whether InNonSandboxFilePath is a file path to InMountPointRoot */
+bool IsFileInMountPoint(const FString& InNonSandboxFilePath, const FSandboxedPlatformFilePath& InMountPointRoot);
+
+/** @return Sandboxed file path if InNonSandboxFilePath is under InMountPointRoot. */
+TOptional<FSandboxedPlatformFilePath> MakeSandboxedFilePathIfInMountPoint(
+	const FString& InNonSandboxFilePath, const FSandboxedPlatformFilePath& InMountPointRoot
+	);
+}
+
diff --git a/Engine/Plugins/Developer/Sandbox/FileSandbox/Source/FileSandboxCore/Private/Sandbox/Platform/SandboxPlatformFile.cpp b/Engine/Plugins/Developer/Sandbox/FileSandbox/Source/FileSandboxCore/Private/Sandbox/Platform/SandboxPlatformFile.cpp
--- a/Engine/Plugins/Developer/Sandbox/FileSandbox/Source/FileSandboxCore/Private/Sandbox/Platform/SandboxPlatformFile.cpp
+++ b/Engine/Plugins/Developer/Sandbox/FileSandbox/Source/FileSandboxCore/Private/Sandbox/Platform/SandboxPlatformFile.cpp
@@ -18,6 +18,7 @@
 #include "Misc/ScopeLock.h"
 #include "Modules/ModuleManager.h"
 #include "RenderingThread.h"
+#include "SanboxedFilePathUtils.h"
 #include "Sandbox/Persist/DiffUtils.h"
 #include "Sandbox/Persist/DiscardUtils.h"
 #include "Sandbox/Persist/PersistUtils.h"
@@ -891,13 +892,9 @@ FSandboxedPlatformFilePath FSandboxPlatformFile::ToSandboxPath_Absolute(FString
 		FScopeLock SandboxMountPointsLock(&SandboxMountPointsCS);
 		for (const FSandboxMountPoint& SandboxMountPoint : SandboxMountPoints)
 		{
-			// Mount point are stored with a trailing slash to prevent matching mount point with similar names -> (/Bla/Content, /Bla/ContentSupreme)
-			// So we test without the slash to make sure with can match mount point directly -> (/Bla/Content match /Bla/Content/)
-			const FString& PathStr = SandboxMountPoint.Path.GetNonSandboxPath();
-			int32 PathStrNoSlashLength = PathStr.Len() - 1;
-			if (FCString::Strnicmp(*InFilename, *PathStr, PathStrNoSlashLength) == 0 && (InFilename.Len() == PathStrNoSlashLength || InFilename[PathStrNoSlashLength] == TEXT('/')))
+			if (const TOptional<FSandboxedPlatformFilePath> Path = MakeSandboxedFilePathIfInMountPoint(InFilename, SandboxMountPoint.Path))
 			{
-				return FSandboxedPlatformFilePath::CreateSandboxPath(MoveTemp(InFilename), SandboxMountPoint.Path);
+				return *Path;
 			}
 		}
 	}
diff --git a/Engine/Plugins/Developer/Sandbox/FileSandbox/Source/FileSandboxCore/Private/Sandbox/SandboxInstance.cpp b/Engine/Plugins/Developer/Sandbox/FileSandbox/Source/FileSandboxCore/Private/Sandbox/SandboxInstance.cpp
--- a/Engine/Plugins/Developer/Sandbox/FileSandbox/Source/FileSandboxCore/Private/Sandbox/SandboxInstance.cpp
+++ b/Engine/Plugins/Developer/Sandbox/FileSandbox/Source/FileSandboxCore/Private/Sandbox/SandboxInstance.cpp
@@ -191,6 +191,11 @@ void FSandboxInstance::EnumerateFileChanges(TFunctionRef<FProcessFileChangeSigna
 	FileSandboxCore::EnumerateFileChanges(*PlatformFile->GetLowerLevel(), GetRootDirectory(), ManifestContent, InProcess);
 }
 
+TOptional<FDateTime> FSandboxInstance::GetSandboxedFileTimestamp(const FString& InFilePath) const
+{
+	return GetSandboxTimestamp(InFilePath, GetRootDirectory(), ManifestContent);
+}
+
 bool FSandboxInstance::DeletedPackageExistsInNonSandbox(const FString& InFilename) const
 {
 	return PlatformFile->DeletedPackageExistsInNonSandbox(InFilename);
diff --git a/Engine/Plugins/Developer/Sandbox/FileSandbox/Source/FileSandboxCore/Private/Sandbox/SandboxInstance.h b/Engine/Plugins/Developer/Sandbox/FileSandbox/Source/FileSandboxCore/Private/Sandbox/SandboxInstance.h
--- a/Engine/Plugins/Developer/Sandbox/FileSandbox/Source/FileSandboxCore/Private/Sandbox/SandboxInstance.h
+++ b/Engine/Plugins/Developer/Sandbox/FileSandbox/Source/FileSandboxCore/Private/Sandbox/SandboxInstance.h
@@ -36,10 +36,14 @@ public:
 
 	using ISandboxInstance::DiscardSandbox;
 
+	const FFileSandboxCore_ManifestData& GetManifestContent() const { return ManifestContent; }
+	IPlatformFile& GetLowerLevelPlatformFile() const { return *PlatformFile->GetLowerLevel(); }
+	
 	//~ Begin ISandboxInstance Interface
 	virtual FPersistResult PersistSandbox(const FPersistArgs& InArgs) override;
 	virtual FDiscardResult DiscardSandbox(const FDiscardArgs& InArgs) override;
 	virtual void EnumerateFileChanges(TFunctionRef<FProcessFileChangeSignature> InProcess) const override;
+	virtual TOptional<FDateTime> GetSandboxedFileTimestamp(const FString& InFilePath) const override;
 	virtual bool DeletedPackageExistsInNonSandbox(const FString& InFilename) const override;
 	virtual const FFileSandboxCore_SandboxMetaData& GetInitialMetaData() const override { return MetaData; }
 	virtual const TCHAR* GetRootDirectory() const override { return *RootDirectory; }
diff --git a/Engine/Plugins/Developer/Sandbox/FileSandbox/Source/FileSandboxCore/Private/Sandbox/SandboxManager.cpp b/Engine/Plugins/Developer/Sandbox/FileSandbox/Source/FileSandboxCore/Private/Sandbox/SandboxManager.cpp
--- a/Engine/Plugins/Developer/Sandbox/FileSandbox/Source/FileSandboxCore/Private/Sandbox/SandboxManager.cpp
+++ b/Engine/Plugins/Developer/Sandbox/FileSandbox/Source/FileSandboxCore/Private/Sandbox/SandboxManager.cpp
@@ -2,12 +2,15 @@
 
 #include "SandboxManager.h"
 
+#include "IFileSandboxCoreModule.h"
+#include "ISandboxRepository.h"
 #include "GenericPlatform/GenericPlatformFile.h"
 #include "HAL/PlatformFileManager.h"
 #include "LogFileSandbox.h"
 #include "Misc/CoreDelegates.h"
 #include "Misc/Paths.h"
 #include "Persist/DiffUtils.h"
+#include "Platform/SanboxedFilePathUtils.h"
 #include "Types/EBreakBehavior.h"
 #include "Types/SandboxFileChange.h"
 #include "Types/Manager/DeleteSandboxByDirectoryArgs.h"
@@ -135,12 +138,16 @@ FLeaveSandboxResult FSandboxManager::LeaveSandbox(const FLeaveSandboxArgs& InLea
 
 void FSandboxManager::EnumerateFileChanges(const FString& InSandboxRootPath, TFunctionRef<FProcessFileChangeSignature> InProcess) const
 {
-	ISandboxInstance* Instance = GetActiveSandboxInstance();
+	FSandboxInstance* Instance = GetSandboxInstanceInternal();
 	if (Instance && FPaths::IsSamePath(Instance->GetRootDirectory(), InSandboxRootPath))
 	{
-		Instance->EnumerateFileChanges(InProcess);
+		FileSandboxCore::EnumerateFileChanges(
+			Instance->GetLowerLevelPlatformFile(), Instance->GetRootDirectory(), Instance->GetManifestContent(), InProcess
+			);
+		return;
 	}
-	else if (const TOptional<FFileSandboxCore_ManifestData> Manifest = LoadManifest(InSandboxRootPath))
+	
+	if (const TOptional<FFileSandboxCore_ManifestData> Manifest = LoadManifest(InSandboxRootPath))
 	{
 		IPlatformFile& PlatformFile = Instance 
 			? *FPlatformFileManager::Get().GetPlatformFile().GetLowerLevel() 
@@ -149,6 +156,18 @@ void FSandboxManager::EnumerateFileChanges(const FString& InSandboxRootPath, TFu
 	}
 }
 
+TOptional<FDateTime> FSandboxManager::GetSandboxedFileTimestamp(const FString& InSandboxRootPath, const FString& InFilePath) const
+{
+	FSandboxInstance* Instance = GetSandboxInstanceInternal();
+	if (Instance && FPaths::IsSamePath(Instance->GetRootDirectory(), InSandboxRootPath))
+	{
+		return GetSandboxTimestamp(InFilePath, InSandboxRootPath, Instance->GetManifestContent());
+	}
+	
+	const TOptional<FFileSandboxCore_ManifestData> Manifest = LoadManifest(InSandboxRootPath);
+	return Manifest ? GetSandboxTimestamp(InFilePath, InSandboxRootPath, *Manifest) : TOptional<FDateTime>();
+}
+
 void FSandboxManager::OnEngineExit()
 {
 	LeaveSandbox();
diff --git a/Engine/Plugins/Developer/Sandbox/FileSandbox/Source/FileSandboxCore/Private/Sandbox/SandboxManager.h b/Engine/Plugins/Developer/Sandbox/FileSandbox/Source/FileSandboxCore/Private/Sandbox/SandboxManager.h
--- a/Engine/Plugins/Developer/Sandbox/FileSandbox/Source/FileSandboxCore/Private/Sandbox/SandboxManager.h
+++ b/Engine/Plugins/Developer/Sandbox/FileSandbox/Source/FileSandboxCore/Private/Sandbox/SandboxManager.h
@@ -24,6 +24,7 @@ public:
 	virtual ISandboxLock* GetActiveLock() const override { return ActiveSandboxData ? ActiveSandboxData->Lock.Get() : nullptr; }
 	virtual ISandboxInstance* GetActiveSandboxInstance() const override { return ActiveSandboxData ? ActiveSandboxData->Sandbox.Get() : nullptr; }
 	virtual void EnumerateFileChanges(const FString& InSandboxRootPath, TFunctionRef<FProcessFileChangeSignature> InProcess) const override;
+	virtual TOptional<FDateTime> GetSandboxedFileTimestamp(const FString& InSandboxRootPath, const FString& InFilePath) const override;
 	virtual FSandboxInstanceEvent& OnPostSandboxStartup() override { return OnPostSandboxStartupDelegate; }
 	virtual FSandboxInstanceEvent& OnPreSandboxShutdown() override { return OnPreSandboxShutdownDelegate; }
 	virtual FSandboxShutdownEvent& OnPostSandboxShutdown() override { return OnPostSandboxShutdownDelegate; }
@@ -56,6 +57,8 @@ private:
 
 	/** Stops the current sandbox. */
 	void OnEngineExit();
+	
+	FSandboxInstance* GetSandboxInstanceInternal() const { return ActiveSandboxData ? ActiveSandboxData->Sandbox.Get() : nullptr;  }
 };
 }
 
diff --git a/Engine/Plugins/Developer/Sandbox/FileSandbox/Source/FileSandboxCore/Private/Utils/SandboxFileUtils.cpp b/Engine/Plugins/Developer/Sandbox/FileSandbox/Source/FileSandboxCore/Private/Utils/SandboxFileUtils.cpp
--- a/Engine/Plugins/Developer/Sandbox/FileSandbox/Source/FileSandboxCore/Private/Utils/SandboxFileUtils.cpp
+++ b/Engine/Plugins/Developer/Sandbox/FileSandbox/Source/FileSandboxCore/Private/Utils/SandboxFileUtils.cpp
@@ -7,7 +7,11 @@
 #include "Data/ManifestData.h"
 #include "Misc/FileHelper.h"
 #include "Misc/Optional.h"
+#include "Misc/PackageName.h"
 #include "Misc/Paths.h"
+#include "Sandbox/Platform/SanboxedFilePathUtils.h"
+#include "Sandbox/Platform/SandboxedPlatformFilePath.h"
+#include "Types/EBreakBehavior.h"
 
 namespace UE::FileSandboxCore
 {
@@ -52,4 +56,53 @@ FString GetSandboxMountPointRoot(const FString& InSandboxRootPath)
 {
 	return InSandboxRootPath / TEXT("Sandbox");
 }
+
+void EnumerateMountPoints(TFunctionRef<EBreakBehavior(const FString& InAssetPath, const FString& InFilesystemPath)> InProcess)
+{
+	TArray<FString> RootPaths;
+	FPackageName::QueryRootContentPaths(RootPaths);
+	for (const FString& AssetPath : RootPaths)
+	{
+		const FString FilesystemPath = FPackageName::LongPackageNameToFilename(AssetPath);
+		if (InProcess(AssetPath, FilesystemPath) == EBreakBehavior::Break)
+		{
+			break;
+		}
+	}	
+}
+
+TOptional<FDateTime> GetSandboxTimestamp(
+	const FString& InRootSandboxPath, 
+	const FString& InNonSandboxPath, const FFileSandboxCore_ManifestData& InManifest
+	)
+{
+	const TArray<FString>& RemovedFiles = InManifest.DeletedFiles.GetFiles();
+	const int32 RemovedIndex = InManifest.DeletedFiles.GetFiles().IndexOfByPredicate([&InNonSandboxPath](const FString& RemovedFilePath)
+	{
+		return FPaths::IsSamePath(InNonSandboxPath, RemovedFilePath);
+	});
+	if (RemovedFiles.IsValidIndex(RemovedIndex))
+	{
+		// The file may have been removed but the user could have tempered with the manifest such that the timestamps are not the same length.
+		const TArray<FDateTime>& Timestamps = InManifest.DeletedFiles.GetTimestamps();
+		return Timestamps.IsValidIndex(RemovedIndex) ? Timestamps[RemovedIndex] : FDateTime::MinValue();
+	}
+	
+	FDateTime EditOrAddTime = FDateTime::MinValue();
+	const auto ProcessMountPoint = [&InRootSandboxPath, &InNonSandboxPath, &EditOrAddTime](const FString& InAssetPath, const FString& InFilesystemPath)
+		{
+			const FSandboxedPlatformFilePath MountPoint = FSandboxedPlatformFilePath::CreateMountPoint(
+				GetSandboxMountPointRoot(InRootSandboxPath), InAssetPath, InFilesystemPath
+			);
+			const TOptional<FSandboxedPlatformFilePath> FilePath = MakeSandboxedFilePathIfInMountPoint(InNonSandboxPath, MountPoint);
+			if (FilePath && ensure(FilePath->HasSandboxPath()))
+			{
+				EditOrAddTime = IFileManager::Get().GetTimeStamp(*FilePath->GetSandboxPath());
+				return EBreakBehavior::Break;
+			}
+			return EBreakBehavior::Continue;
+		};
+	EnumerateMountPoints(ProcessMountPoint);
+	return EditOrAddTime;
+}
 }
diff --git a/Engine/Plugins/Developer/Sandbox/FileSandbox/Source/FileSandboxCore/Private/Utils/SandboxFileUtils.h b/Engine/Plugins/Developer/Sandbox/FileSandbox/Source/FileSandboxCore/Private/Utils/SandboxFileUtils.h
--- a/Engine/Plugins/Developer/Sandbox/FileSandbox/Source/FileSandboxCore/Private/Utils/SandboxFileUtils.h
+++ b/Engine/Plugins/Developer/Sandbox/FileSandbox/Source/FileSandboxCore/Private/Utils/SandboxFileUtils.h
@@ -2,7 +2,12 @@
 
 #pragma once
 
+#include "Templates/FunctionFwd.h"
+#include "HAL/Platform.h"
+
 class FString;
+enum class EBreakBehavior : uint8;
+struct FDateTime;
 struct FFileSandboxCore_ManifestData;
 template<typename OptionalType> struct TOptional;
 
@@ -20,4 +25,10 @@ FString GetMetadataFileName();
 
 /** @return Path to directory in the sandbox file structure where mount points are stored, i.e. "MySandbox/Sandbox" */
 FString GetSandboxMountPointRoot(const FString& InSandboxRootPath);
+
+/** Enumerates all mount points in the engine. */
+void EnumerateMountPoints(TFunctionRef<EBreakBehavior(const FString& InAssetPath, const FString& InFilesystemPath)> InProcess);
+
+/** @return Timestamp in local time when InNonSandboxPath was edited, added, or removed by the sandbox. */
+TOptional<FDateTime> GetSandboxTimestamp(const FString& InNonSandboxPath, const FString& InRootSandboxPath, const FFileSandboxCore_ManifestData& InManifest);
 }
\ No newline at end of file
diff --git a/Engine/Plugins/Developer/Sandbox/FileSandbox/Source/FileSandboxCore/Public/ISandboxInstance.h b/Engine/Plugins/Developer/Sandbox/FileSandbox/Source/FileSandboxCore/Public/ISandboxInstance.h
--- a/Engine/Plugins/Developer/Sandbox/FileSandbox/Source/FileSandboxCore/Public/ISandboxInstance.h
+++ b/Engine/Plugins/Developer/Sandbox/FileSandbox/Source/FileSandboxCore/Public/ISandboxInstance.h
@@ -72,6 +72,9 @@ public:
 	 * @see GetDirtyPackages
 	 */
 	UE_API bool HasFileChanges() const;
+	
+	/** @return The timestamp at which the file was last modified by the sandbox, or FDateTime::MinValue(). */
+	virtual TOptional<FDateTime> GetSandboxedFileTimestamp(const FString& InFilePath) const = 0;
 
 	
 	/** @return Whether the given package file exists on non sandbox path, i.e. whether the file was deleted in sandbox mode. */
diff --git a/Engine/Plugins/Developer/Sandbox/FileSandbox/Source/FileSandboxCore/Public/ISandboxManager.h b/Engine/Plugins/Developer/Sandbox/FileSandbox/Source/FileSandboxCore/Public/ISandboxManager.h
--- a/Engine/Plugins/Developer/Sandbox/FileSandbox/Source/FileSandboxCore/Public/ISandboxManager.h
+++ b/Engine/Plugins/Developer/Sandbox/FileSandbox/Source/FileSandboxCore/Public/ISandboxManager.h
@@ -4,6 +4,7 @@
 
 #include "Delegates/Delegate.h"
 #include "Types/ArgsFwds.h"
+#include "Types/GatheredFileChanges.h"
 #include "Types/Manager/LeaveSandboxArgs.h"
 #include "Types/Manager/LeaveSandboxResult.h"
 #include "Types/Manager/SandboxCreationResult.h"
@@ -144,6 +145,9 @@ public:
 	/** Alternative version that uses the sandbox root directory instead.*/
 	virtual void EnumerateFileChanges(const FString& InSandboxRootPath, TFunctionRef<FProcessFileChangeSignature> InProcess) const = 0;
 	
+	/** @return The timestamp at which the file was last modified by the sandbox, or FDateTime::MinValue(). */
+	virtual TOptional<FDateTime> GetSandboxedFileTimestamp(const FString& InSandboxRootPath, const FString& InFilePath) const = 0;
+	
 	/** @return Event invoked after a sandbox has been successfully created. */
 	virtual FSandboxInstanceEvent& OnPostSandboxStartup() = 0;
 	
diff --git a/Engine/Plugins/Developer/Sandbox/FileSandbox/Source/FileSandboxCore/Public/Types/GatheredFileChanges.h b/Engine/Plugins/Developer/Sandbox/FileSandbox/Source/FileSandboxCore/Public/Types/GatheredFileChanges.h
--- a/Engine/Plugins/Developer/Sandbox/FileSandbox/Source/FileSandboxCore/Public/Types/GatheredFileChanges.h
+++ b/Engine/Plugins/Developer/Sandbox/FileSandbox/Source/FileSandboxCore/Public/Types/GatheredFileChanges.h
@@ -16,7 +16,8 @@ enum class EFileChangeGatherFlags : uint8
 {
 	None,
 	IncludeChangeTypes = 1 << 0,
-	All = IncludeChangeTypes
+	IncludeTimestamps = 1 << 0,
+	All = IncludeChangeTypes | IncludeTimestamps
 };
 ENUM_CLASS_FLAGS(EFileChangeGatherFlags);
 
@@ -29,6 +30,9 @@ struct FGatheredFileChanges
 	/** The changes made to the file. Equal length to NonSandboxPaths. Only set if EFileChangeGatherFlags::IncludeChangeTypes was set. */
 	TArray<ESandboxFileChange> FileActions;
 	
+	/** When the change was made, in local time. Only set if EFileChangeGatherFlags::IncludeTimestamps was set.*/
+	TArray<FDateTime> Timestamps;
+	
 	/** @return Whether there are any changes. */
 	bool HasChanges() const { return !NonSandboxPaths.IsEmpty(); }
 };
