File Sandbox Core: Fix non-unity compile issues

#Synced-CL 49159615

diff --git a/Engine/Plugins/Developer/Sandbox/FileSandbox/Source/FileSandboxCore/Public/Types/GatheredFileChanges.h b/Engine/Plugins/Developer/Sandbox/FileSandbox/Source/FileSandboxCore/Public/Types/GatheredFileChanges.h
--- a/Engine/Plugins/Developer/Sandbox/FileSandbox/Source/FileSandboxCore/Public/Types/GatheredFileChanges.h
+++ b/Engine/Plugins/Developer/Sandbox/FileSandbox/Source/FileSandboxCore/Public/Types/GatheredFileChanges.h
@@ -4,6 +4,7 @@
 
 #include "Containers/Array.h"
 #include "HAL/Platform.h"
+#include "Misc/DateTime.h"
 
 class FString;
 
