File Sandbox Core: Introduce FSandboxedFileChangeInfo and refactor EnumerateFileChanges to use that struct to pass all file change info.

This is more forward-compatible. It allows is more performant because ISandboxManager::EnumerateFileChanges calls LoadManifest once as opposed to forcing the API user to call GetSandboxedFileTimestamp separately for each file: GetSandboxedFileTimestamp would cause LoadManifest for each call. This also allows us to retain FFileSandboxCore_ManifestData in the private interface without exposing it to the API user.

#Synced-CL 49159615

diff --git a/Engine/Plugins/Developer/Sandbox/FileSandbox/Source/FileSandboxCore/Private/Sandbox/ISandboxInstance.cpp b/Engine/Plugins/Developer/Sandbox/FileSandbox/Source/FileSandboxCore/Private/Sandbox/ISandboxInstance.cpp
--- a/Engine/Plugins/Developer/Sandbox/FileSandbox/Source/FileSandboxCore/Private/Sandbox/ISandboxInstance.cpp
+++ b/Engine/Plugins/Developer/Sandbox/FileSandbox/Source/FileSandboxCore/Private/Sandbox/ISandboxInstance.cpp
@@ -28,27 +28,28 @@ void ISandboxInstance::DiscardAll()
 FGatheredFileChanges ISandboxInstance::GatherChangedFiles(EFileChangeGatherFlags InFlags) const
 {
 	FGatheredFileChanges Result;
-	EnumerateFileChanges([this, InFlags, &Result](const FString& InFile, ESandboxFileChange InChange)
+	const EFileEnumerationFlags EnumerationFlags = EnumHasAnyFlags(InFlags, EFileChangeGatherFlags::IncludeTimestamps) 
+		? EFileEnumerationFlags::IncludeTimestamps : EFileEnumerationFlags::None;
+	EnumerateFileChanges([this, InFlags, &Result](const FSandboxedFileChangeInfo& InChange)
 	{
-		Result.NonSandboxPaths.Add(InFile);
+		Result.NonSandboxPaths.Add(InChange.Path);
 		if (EnumHasAnyFlags(InFlags, EFileChangeGatherFlags::IncludeChangeTypes))
 		{
-			Result.FileActions.Add(InChange);
+			Result.FileActions.Add(InChange.Action);
 		}
 		if (EnumHasAnyFlags(InFlags, EFileChangeGatherFlags::IncludeTimestamps))
 		{
-			const TOptional<FDateTime> Timestamp = GetSandboxedFileTimestamp(InFile);
-			Result.Timestamps.Add(Timestamp.Get(FDateTime::MinValue()));
+			Result.Timestamps.Add(InChange.Timestamp);
 		}
 		return EBreakBehavior::Continue;
-	});
+	}, EnumerationFlags);
 	return Result;
 }
 
 bool ISandboxInstance::HasFileChanges() const
 {
 	bool bHasChanges = false;
-	EnumerateFileChanges([&bHasChanges](const FString& InFile, ESandboxFileChange InChange)
+	EnumerateFileChanges([&bHasChanges](const FSandboxedFileChangeInfo& InChange)
 	{
 		bHasChanges = true;
 		return EBreakBehavior::Break;
diff --git a/Engine/Plugins/Developer/Sandbox/FileSandbox/Source/FileSandboxCore/Private/Sandbox/Persist/DiffUtils.cpp b/Engine/Plugins/Developer/Sandbox/FileSandbox/Source/FileSandboxCore/Private/Sandbox/Persist/DiffUtils.cpp
--- a/Engine/Plugins/Developer/Sandbox/FileSandbox/Source/FileSandboxCore/Private/Sandbox/Persist/DiffUtils.cpp
+++ b/Engine/Plugins/Developer/Sandbox/FileSandbox/Source/FileSandboxCore/Private/Sandbox/Persist/DiffUtils.cpp
@@ -106,13 +106,24 @@ void EnumerateFileChanges(
 	IPlatformFile& InPlatformFile,
 	const FString& InSandboxRootPath, 
 	const FFileSandboxCore_ManifestData& InManifest,
-	TFunctionRef<FProcessFileChangeSignature> InProcess
+	TFunctionRef<FProcessFileChangeSignature> InProcess,
+	EFileEnumerationFlags InFlags
 	)
 {
+	const auto ProcessChange = [&InSandboxRootPath, &InManifest, &InProcess, InFlags](const FString& InNonSandboxPath, ESandboxFileChange InAction)
+	{
+		FSandboxedFileChangeInfo Info{ InNonSandboxPath, InAction };
+		if (EnumHasAnyFlags(InFlags, EFileEnumerationFlags::IncludeTimestamps))
+		{
+			Info.Timestamp = GetSandboxTimestamp(InNonSandboxPath, InSandboxRootPath, InManifest).Get(FDateTime::MinValue());
+		}
+		return InProcess(Info);
+	};
+	
 	EBreakBehavior BreakBehavior = EBreakBehavior::Continue;
-	AnalyzeFilesMarkedRemoved(InPlatformFile, InManifest.DeletedFiles.GetFiles(), [&InProcess, &BreakBehavior](const FString& InFile)
+	AnalyzeFilesMarkedRemoved(InPlatformFile, InManifest.DeletedFiles.GetFiles(), [&ProcessChange, &BreakBehavior](const FString& InFile)
 	{
-		BreakBehavior = InProcess(InFile, ESandboxFileChange::Removed);
+		BreakBehavior = ProcessChange(InFile, ESandboxFileChange::Removed);
 		return BreakBehavior;
 	});
 	if (BreakBehavior == EBreakBehavior::Break)
@@ -120,21 +131,21 @@ void EnumerateFileChanges(
 		return;
 	}
 	
-	const auto ProcessMountPoint = [&InSandboxRootPath, &InProcess, &InPlatformFile, &BreakBehavior]
+	const auto ProcessMountPoint = [&InSandboxRootPath, &ProcessChange, &InPlatformFile, &BreakBehavior]
 		(const FString& InAssetPath, const FString& InFilesystemPath)
 	{
 		const FSandboxedPlatformFilePath MountPoint = FSandboxedPlatformFilePath::CreateMountPoint(
 			GetSandboxMountPointRoot(InSandboxRootPath), InAssetPath, InFilesystemPath
 		);
 			
-		const auto HandleEdited = [&InProcess, &BreakBehavior](const FSandboxedPlatformFilePath& InPath)
+		const auto HandleEdited = [&ProcessChange, &BreakBehavior](const FSandboxedPlatformFilePath& InPath)
 		{
-			BreakBehavior = InProcess(InPath.GetNonSandboxPath(), ESandboxFileChange::Edited);
+			BreakBehavior = ProcessChange(InPath.GetNonSandboxPath(), ESandboxFileChange::Edited);
 			return BreakBehavior;
 		};
-		const auto HandleAdded = [&InProcess, &BreakBehavior](const FSandboxedPlatformFilePath& InPath)
+		const auto HandleAdded = [&ProcessChange, &BreakBehavior](const FSandboxedPlatformFilePath& InPath)
 		{
-			BreakBehavior = InProcess(InPath.GetNonSandboxPath(), ESandboxFileChange::Added);
+			BreakBehavior = ProcessChange(InPath.GetNonSandboxPath(), ESandboxFileChange::Added);
 			return BreakBehavior;
 		};
 			
diff --git a/Engine/Plugins/Developer/Sandbox/FileSandbox/Source/FileSandboxCore/Private/Sandbox/Persist/DiffUtils.h b/Engine/Plugins/Developer/Sandbox/FileSandbox/Source/FileSandboxCore/Private/Sandbox/Persist/DiffUtils.h
--- a/Engine/Plugins/Developer/Sandbox/FileSandbox/Source/FileSandboxCore/Private/Sandbox/Persist/DiffUtils.h
+++ b/Engine/Plugins/Developer/Sandbox/FileSandbox/Source/FileSandboxCore/Private/Sandbox/Persist/DiffUtils.h
@@ -72,6 +72,7 @@ void ScanForEditedOrAddedFiles(
  */
 void EnumerateFileChanges(
 	IPlatformFile& InPlatformFile, const FString& InSandboxRootPath, const FFileSandboxCore_ManifestData& InManifest, 
-	TFunctionRef<FProcessFileChangeSignature> InProcess
+	TFunctionRef<EBreakBehavior(const FSandboxedFileChangeInfo& InFileInfo)> InProcess,
+	EFileEnumerationFlags InFlags = EFileEnumerationFlags::All
 	);
 }
diff --git a/Engine/Plugins/Developer/Sandbox/FileSandbox/Source/FileSandboxCore/Private/Sandbox/SandboxInstance.cpp b/Engine/Plugins/Developer/Sandbox/FileSandbox/Source/FileSandboxCore/Private/Sandbox/SandboxInstance.cpp
--- a/Engine/Plugins/Developer/Sandbox/FileSandbox/Source/FileSandboxCore/Private/Sandbox/SandboxInstance.cpp
+++ b/Engine/Plugins/Developer/Sandbox/FileSandbox/Source/FileSandboxCore/Private/Sandbox/SandboxInstance.cpp
@@ -186,9 +186,9 @@ FDiscardResult FSandboxInstance::DiscardSandbox(const FDiscardArgs& InArgs)
 	return Result;
 }
 
-void FSandboxInstance::EnumerateFileChanges(TFunctionRef<FProcessFileChangeSignature> InProcess) const
+void FSandboxInstance::EnumerateFileChanges(TFunctionRef<FProcessFileChangeSignature> InProcess, EFileEnumerationFlags InFlags) const
 {
-	FileSandboxCore::EnumerateFileChanges(*PlatformFile->GetLowerLevel(), GetRootDirectory(), ManifestContent, InProcess);
+	FileSandboxCore::EnumerateFileChanges(*PlatformFile->GetLowerLevel(), GetRootDirectory(), ManifestContent, InProcess, InFlags);
 }
 
 TOptional<FDateTime> FSandboxInstance::GetSandboxedFileTimestamp(const FString& InFilePath) const
@@ -252,11 +252,11 @@ void FSandboxInstance::OnEndOfFrame()
 	{
 		FFileSandboxCore_ManifestData NewManifest;
 		NewManifest.VersionInfo.Initialize();
-		EnumerateFileChanges([&NewManifest](const FString& InPath, ESandboxFileChange InChange)
+		EnumerateFileChanges([&NewManifest](const FSandboxedFileChangeInfo& Change)
 		{
-			Private::TrackFileChangeInManifest(NewManifest, InPath, InChange);
+			Private::TrackFileChangeInManifest(NewManifest, Change.Path, Change.Action);
 			return EBreakBehavior::Continue;
-		});
+		}, EFileEnumerationFlags::None);
 		
 		ManifestContent = MoveTemp(NewManifest);
 		SaveManifest();
diff --git a/Engine/Plugins/Developer/Sandbox/FileSandbox/Source/FileSandboxCore/Private/Sandbox/SandboxInstance.h b/Engine/Plugins/Developer/Sandbox/FileSandbox/Source/FileSandboxCore/Private/Sandbox/SandboxInstance.h
--- a/Engine/Plugins/Developer/Sandbox/FileSandbox/Source/FileSandboxCore/Private/Sandbox/SandboxInstance.h
+++ b/Engine/Plugins/Developer/Sandbox/FileSandbox/Source/FileSandboxCore/Private/Sandbox/SandboxInstance.h
@@ -42,7 +42,7 @@ public:
 	//~ Begin ISandboxInstance Interface
 	virtual FPersistResult PersistSandbox(const FPersistArgs& InArgs) override;
 	virtual FDiscardResult DiscardSandbox(const FDiscardArgs& InArgs) override;
-	virtual void EnumerateFileChanges(TFunctionRef<FProcessFileChangeSignature> InProcess) const override;
+	virtual void EnumerateFileChanges(TFunctionRef<FProcessFileChangeSignature> InProcess, EFileEnumerationFlags InFlags) const override;
 	virtual TOptional<FDateTime> GetSandboxedFileTimestamp(const FString& InFilePath) const override;
 	virtual bool DeletedPackageExistsInNonSandbox(const FString& InFilename) const override;
 	virtual const FFileSandboxCore_SandboxMetaData& GetInitialMetaData() const override { return MetaData; }
diff --git a/Engine/Plugins/Developer/Sandbox/FileSandbox/Source/FileSandboxCore/Private/Sandbox/SandboxManager.cpp b/Engine/Plugins/Developer/Sandbox/FileSandbox/Source/FileSandboxCore/Private/Sandbox/SandboxManager.cpp
--- a/Engine/Plugins/Developer/Sandbox/FileSandbox/Source/FileSandboxCore/Private/Sandbox/SandboxManager.cpp
+++ b/Engine/Plugins/Developer/Sandbox/FileSandbox/Source/FileSandboxCore/Private/Sandbox/SandboxManager.cpp
@@ -136,13 +136,13 @@ FLeaveSandboxResult FSandboxManager::LeaveSandbox(const FLeaveSandboxArgs& InLea
 	return ELeaveSandboxErrorCode::Success;
 }
 
-void FSandboxManager::EnumerateFileChanges(const FString& InSandboxRootPath, TFunctionRef<FProcessFileChangeSignature> InProcess) const
+void FSandboxManager::EnumerateFileChanges(const FString& InSandboxRootPath, TFunctionRef<FProcessFileChangeSignature> InProcess, EFileEnumerationFlags InFlags) const
 {
 	FSandboxInstance* Instance = GetSandboxInstanceInternal();
 	if (Instance && FPaths::IsSamePath(Instance->GetRootDirectory(), InSandboxRootPath))
 	{
 		FileSandboxCore::EnumerateFileChanges(
-			Instance->GetLowerLevelPlatformFile(), Instance->GetRootDirectory(), Instance->GetManifestContent(), InProcess
+			Instance->GetLowerLevelPlatformFile(), Instance->GetRootDirectory(), Instance->GetManifestContent(), InProcess, InFlags
 			);
 		return;
 	}
@@ -152,7 +152,7 @@ void FSandboxManager::EnumerateFileChanges(const FString& InSandboxRootPath, TFu
 		IPlatformFile& PlatformFile = Instance 
 			? *FPlatformFileManager::Get().GetPlatformFile().GetLowerLevel() 
 			: FPlatformFileManager::Get().GetPlatformFile();
-		FileSandboxCore::EnumerateFileChanges(PlatformFile, InSandboxRootPath, *Manifest, InProcess);
+		FileSandboxCore::EnumerateFileChanges(PlatformFile, InSandboxRootPath, *Manifest, InProcess, InFlags);
 	}
 }
 
diff --git a/Engine/Plugins/Developer/Sandbox/FileSandbox/Source/FileSandboxCore/Private/Sandbox/SandboxManager.h b/Engine/Plugins/Developer/Sandbox/FileSandbox/Source/FileSandboxCore/Private/Sandbox/SandboxManager.h
--- a/Engine/Plugins/Developer/Sandbox/FileSandbox/Source/FileSandboxCore/Private/Sandbox/SandboxManager.h
+++ b/Engine/Plugins/Developer/Sandbox/FileSandbox/Source/FileSandboxCore/Private/Sandbox/SandboxManager.h
@@ -23,7 +23,7 @@ public:
 	virtual FLeaveSandboxResult LeaveSandbox(const FLeaveSandboxArgs& InLeaveArgs = FLeaveSandboxArgs()) override;
 	virtual ISandboxLock* GetActiveLock() const override { return ActiveSandboxData ? ActiveSandboxData->Lock.Get() : nullptr; }
 	virtual ISandboxInstance* GetActiveSandboxInstance() const override { return ActiveSandboxData ? ActiveSandboxData->Sandbox.Get() : nullptr; }
-	virtual void EnumerateFileChanges(const FString& InSandboxRootPath, TFunctionRef<FProcessFileChangeSignature> InProcess) const override;
+	virtual void EnumerateFileChanges(const FString& InSandboxRootPath, TFunctionRef<FProcessFileChangeSignature> InProcess, EFileEnumerationFlags InFlags) const override;
 	virtual TOptional<FDateTime> GetSandboxedFileTimestamp(const FString& InSandboxRootPath, const FString& InFilePath) const override;
 	virtual FSandboxInstanceEvent& OnPostSandboxStartup() override { return OnPostSandboxStartupDelegate; }
 	virtual FSandboxInstanceEvent& OnPreSandboxShutdown() override { return OnPreSandboxShutdownDelegate; }
diff --git a/Engine/Plugins/Developer/Sandbox/FileSandbox/Source/FileSandboxCore/Public/ISandboxInstance.h b/Engine/Plugins/Developer/Sandbox/FileSandbox/Source/FileSandboxCore/Public/ISandboxInstance.h
--- a/Engine/Plugins/Developer/Sandbox/FileSandbox/Source/FileSandboxCore/Public/ISandboxInstance.h
+++ b/Engine/Plugins/Developer/Sandbox/FileSandbox/Source/FileSandboxCore/Public/ISandboxInstance.h
@@ -6,6 +6,7 @@
 #include "Types/Sandbox/DiscardArgs.h"
 #include "Types/Sandbox/PersistResult.h"
 #include "Types/ArgsFwds.h"
+#include "Types/SandboxedFileChangeInfo.h"
 
 #define UE_API FILESANDBOXCORE_API
 
@@ -64,7 +65,9 @@ public:
 	 * Invokes InProcess for every effective file change that would be performed if the sandbox was persisted.
 	 * @see ISandboxManager::EnumerateFileChanges.
 	 */
-	virtual void EnumerateFileChanges(TFunctionRef<FProcessFileChangeSignature> InProcess) const = 0;
+	virtual void EnumerateFileChanges(
+		TFunctionRef<FProcessFileChangeSignature> InProcess, EFileEnumerationFlags InFlags = EFileEnumerationFlags::All
+		) const = 0;
 	
 	/** 
 	 * @return Whether any changes have to any files.
diff --git a/Engine/Plugins/Developer/Sandbox/FileSandbox/Source/FileSandboxCore/Public/ISandboxManager.h b/Engine/Plugins/Developer/Sandbox/FileSandbox/Source/FileSandboxCore/Public/ISandboxManager.h
--- a/Engine/Plugins/Developer/Sandbox/FileSandbox/Source/FileSandboxCore/Public/ISandboxManager.h
+++ b/Engine/Plugins/Developer/Sandbox/FileSandbox/Source/FileSandboxCore/Public/ISandboxManager.h
@@ -8,6 +8,7 @@
 #include "Types/Manager/LeaveSandboxArgs.h"
 #include "Types/Manager/LeaveSandboxResult.h"
 #include "Types/Manager/SandboxCreationResult.h"
+#include "Types/SandboxedFileChangeInfo.h"
 #include "Utils/SandboxDirectoryUtils.h"
 
 #define UE_API FILESANDBOXCORE_API
@@ -143,9 +144,12 @@ public:
 		) const;
 	
 	/** Alternative version that uses the sandbox root directory instead.*/
-	virtual void EnumerateFileChanges(const FString& InSandboxRootPath, TFunctionRef<FProcessFileChangeSignature> InProcess) const = 0;
+	virtual void EnumerateFileChanges(
+		const FString& InSandboxRootPath, TFunctionRef<FProcessFileChangeSignature> InProcess,
+		EFileEnumerationFlags InFlags = EFileEnumerationFlags::All
+		) const = 0;
 	
-	/** @return The timestamp at which the file was last modified by the sandbox, or FDateTime::MinValue(). */
+	/** @return The timestamp at which the file was last modified by the sandbox. */
 	virtual TOptional<FDateTime> GetSandboxedFileTimestamp(const FString& InSandboxRootPath, const FString& InFilePath) const = 0;
 	
 	/** @return Event invoked after a sandbox has been successfully created. */
diff --git a/Engine/Plugins/Developer/Sandbox/FileSandbox/Source/FileSandboxCore/Public/Types/ArgsFwds.h b/Engine/Plugins/Developer/Sandbox/FileSandbox/Source/FileSandboxCore/Public/Types/ArgsFwds.h
--- a/Engine/Plugins/Developer/Sandbox/FileSandbox/Source/FileSandboxCore/Public/Types/ArgsFwds.h
+++ b/Engine/Plugins/Developer/Sandbox/FileSandbox/Source/FileSandboxCore/Public/Types/ArgsFwds.h
@@ -19,6 +19,7 @@ struct FNewSandboxArgs;
 struct FLoadSandboxError;
 struct FNewSandboxError;
 struct FSandboxCreationResult;
+struct FSandboxedFileChangeInfo;
 	
 template<typename TErrorType>
 struct TSandboxCreationResult;
@@ -28,5 +29,5 @@ using FLoadSandboxResult = TSandboxCreationResult<FLoadSandboxError>;
 enum class EBreakBehavior : uint8;
 enum class ESandboxFileChange : uint8;
 
-using FProcessFileChangeSignature = EBreakBehavior(const FString& InPath, ESandboxFileChange InChange);
+using FProcessFileChangeSignature = EBreakBehavior(const FSandboxedFileChangeInfo& InFileInfo);
 }
\ No newline at end of file
diff --git a/Engine/Plugins/Developer/Sandbox/FileSandbox/Source/FileSandboxCore/Public/Types/GatheredFileChanges.h b/Engine/Plugins/Developer/Sandbox/FileSandbox/Source/FileSandboxCore/Public/Types/GatheredFileChanges.h
--- a/Engine/Plugins/Developer/Sandbox/FileSandbox/Source/FileSandboxCore/Public/Types/GatheredFileChanges.h
+++ b/Engine/Plugins/Developer/Sandbox/FileSandbox/Source/FileSandboxCore/Public/Types/GatheredFileChanges.h
@@ -17,7 +17,7 @@ enum class EFileChangeGatherFlags : uint8
 {
 	None,
 	IncludeChangeTypes = 1 << 0,
-	IncludeTimestamps = 1 << 0,
+	IncludeTimestamps = 1 << 1,
 	All = IncludeChangeTypes | IncludeTimestamps
 };
 ENUM_CLASS_FLAGS(EFileChangeGatherFlags);
diff --git a/Engine/Plugins/Developer/Sandbox/FileSandbox/Source/FileSandboxCore/Public/Types/SandboxedFileChangeInfo.h b/Engine/Plugins/Developer/Sandbox/FileSandbox/Source/FileSandboxCore/Public/Types/SandboxedFileChangeInfo.h
new file mode 100644
--- /dev/null
+++ b/Engine/Plugins/Developer/Sandbox/FileSandbox/Source/FileSandboxCore/Public/Types/SandboxedFileChangeInfo.h
@@ -0,0 +1,42 @@
+﻿// Copyright Epic Games, Inc. All Rights Reserved.
+
+#pragma once
+
+#include "Containers/UnrealString.h"
+#include "HAL/Platform.h"
+#include "Misc/EnumClassFlags.h"
+#include "Misc/DateTime.h"
+#include "SandboxFileChange.h"
+
+namespace UE::FileSandboxCore
+{
+/**
+ * Flags enumerating file changes.
+ * @see ISandboxManager::EnumerateFileChanges
+ * @see ISandboxInstance::EnumerateFileChanges
+ */
+enum class EFileEnumerationFlags : uint8
+{
+	None,
+	IncludeTimestamps = 1 << 0,
+	All = IncludeTimestamps
+};
+ENUM_CLASS_FLAGS(EFileEnumerationFlags);
+
+/**
+ * Info about each file change enumerated.
+ * @see ISandboxManager::EnumerateFileChanges
+ * @see ISandboxInstance::EnumerateFileChanges
+ */
+struct FSandboxedFileChangeInfo
+{
+	/** Absolute non-sandbox file path. */
+	FString Path;
+	
+	/** The action performed on the file */
+	ESandboxFileChange Action = ESandboxFileChange::None;
+	
+	/** Only set if EFileEnumerationFlags::IncludeTimestamps was specified */
+	FDateTime Timestamp = FDateTime::MinValue();
+};
+}
diff --git a/Engine/Plugins/Developer/Sandbox/SandboxedEditing/Source/SandboxedEditing/Private/Features/Browser/ViewModels/FileState/Models/ActiveSandboxFileStateViewModel.cpp b/Engine/Plugins/Developer/Sandbox/SandboxedEditing/Source/SandboxedEditing/Private/Features/Browser/ViewModels/FileState/Models/ActiveSandboxFileStateViewModel.cpp
--- a/Engine/Plugins/Developer/Sandbox/SandboxedEditing/Source/SandboxedEditing/Private/Features/Browser/ViewModels/FileState/Models/ActiveSandboxFileStateViewModel.cpp
+++ b/Engine/Plugins/Developer/Sandbox/SandboxedEditing/Source/SandboxedEditing/Private/Features/Browser/ViewModels/FileState/Models/ActiveSandboxFileStateViewModel.cpp
@@ -6,6 +6,7 @@
 #include "Framework/Models/SandboxSystemModel.h"
 #include "Misc/Paths.h"
 #include "Types/EBreakBehavior.h"
+#include "Types/SandboxedFileChangeInfo.h"
 
 namespace UE::SandboxedEditing
 {
@@ -52,17 +53,20 @@ void FActiveSandboxFileStateViewModel::RefreshItems()
 	TArray<TSharedPtr<FFileStateItem>> OldItems = MoveTemp(AllUnfilteredItems);
 	const FString ActiveSandbox = Model->GetActiveSandboxPath();
 	Model->EnumerateFileChanges(ActiveSandbox, 
-		[this, &bMadeChanges, &OldItems](const FString& InNonSandboxPath, FileSandboxCore::ESandboxFileChange InAction)
+		[this, &bMadeChanges, &OldItems](const FileSandboxCore::FSandboxedFileChangeInfo& InChange)
 		{
-			const TSharedPtr<FFileStateItem>* ReusedItem = OldItems.FindByPredicate([InNonSandboxPath, InAction](const TSharedPtr<FFileStateItem>& Item)
+			const FString& NonSandboxPath = InChange.Path;
+			const FileSandboxCore::ESandboxFileChange Action = InChange.Action;
+			const FDateTime Timestamp = InChange.Timestamp;
+				
+			const TSharedPtr<FFileStateItem>* ReusedItem = OldItems.FindByPredicate([&NonSandboxPath](const TSharedPtr<FFileStateItem>& Item)
 			{
-				return FPaths::IsSamePath(Item->NonSandboxFile, InNonSandboxPath);
+				return FPaths::IsSamePath(Item->NonSandboxFile, NonSandboxPath);
 			});
-			const FDateTime Timestamp = Model->GetSandboxedFileTimestamp(InNonSandboxPath);
 			if (ReusedItem)
 			{
 				// Do this to make sure the item is up to date
-				ReusedItem->Get()->Action = InAction; 
+				ReusedItem->Get()->Action = Action; 
 				ReusedItem->Get()->Timestamp = Timestamp;
 				
 				AllUnfilteredItems.Emplace(*ReusedItem);
@@ -70,7 +74,7 @@ void FActiveSandboxFileStateViewModel::RefreshItems()
 			else
 			{
 				bMadeChanges = true;
-				AllUnfilteredItems.Emplace(MakeShared<FFileStateItem>(InNonSandboxPath, InAction, Timestamp));
+				AllUnfilteredItems.Emplace(MakeShared<FFileStateItem>(NonSandboxPath, Action, Timestamp));
 			}
 				
 			return FileSandboxCore::EBreakBehavior::Continue;
diff --git a/Engine/Plugins/Developer/Sandbox/SandboxedEditing/Source/SandboxedEditing/Private/Framework/Models/SandboxSystemModel.cpp b/Engine/Plugins/Developer/Sandbox/SandboxedEditing/Source/SandboxedEditing/Private/Framework/Models/SandboxSystemModel.cpp
--- a/Engine/Plugins/Developer/Sandbox/SandboxedEditing/Source/SandboxedEditing/Private/Framework/Models/SandboxSystemModel.cpp
+++ b/Engine/Plugins/Developer/Sandbox/SandboxedEditing/Source/SandboxedEditing/Private/Framework/Models/SandboxSystemModel.cpp
@@ -270,19 +270,19 @@ FileSandboxCore::FGatheredFileChanges FSandboxSystemModel::GatherFileChanges(con
 	ISandboxManager& SandboxManager = Module.GetSandboxManager();
 	
 	FGatheredFileChanges Result;
-	SandboxManager.EnumerateFileChanges(InSandbox, [&InSandbox, &SandboxManager, &Result](const FString& InPath, ESandboxFileChange InChange)
+	SandboxManager.EnumerateFileChanges(InSandbox, [&Result](const FSandboxedFileChangeInfo& InChange)
 	{
-		Result.NonSandboxPaths.Add(InPath);
-		Result.FileActions.Add(InChange);
-		Result.Timestamps.Add(SandboxManager.GetSandboxedFileTimestamp(InSandbox, InPath).Get(FDateTime::MinValue()));
+		Result.NonSandboxPaths.Add(InChange.Path);
+		Result.FileActions.Add(InChange.Action);
+		Result.Timestamps.Add(InChange.Timestamp);
 		return FileSandboxCore::EBreakBehavior::Continue;
-	});
+	}, EFileEnumerationFlags::IncludeTimestamps);
 	return Result;
 }
 
 void FSandboxSystemModel::EnumerateFileChanges(
 	const FString& InSandbox,
-	TFunctionRef<FileSandboxCore::EBreakBehavior(const FString& InPath, FileSandboxCore::ESandboxFileChange InChange)> InCallback
+	TFunctionRef<FileSandboxCore::EBreakBehavior(const FileSandboxCore::FSandboxedFileChangeInfo& InChange)> InCallback
 	)
 {
 	using namespace FileSandboxCore;
@@ -291,19 +291,6 @@ void FSandboxSystemModel::EnumerateFileChanges(
 	SandboxManager.EnumerateFileChanges(InSandbox, InCallback);
 }
 
-FDateTime FSandboxSystemModel::GetSandboxedFileTimestamp(const FString& InPath) const
-{
-	using namespace FileSandboxCore;
-	IFileSandboxCoreModule& Module = IFileSandboxCoreModule::Get();
-	ISandboxManager& SandboxManager = Module.GetSandboxManager();
-	ISandboxInstance* Instance = SandboxManager.GetActiveSandboxInstance();
-	
-	const FDateTime MinValue = FDateTime::MinValue();
-	return Instance 
-		? Instance->GetSandboxedFileTimestamp(InPath).Get(MinValue) 
-		: MinValue;
-}
-
 void FSandboxSystemModel::DiscardAllChanges() const
 {
 	using namespace FileSandboxCore;
diff --git a/Engine/Plugins/Developer/Sandbox/SandboxedEditing/Source/SandboxedEditing/Private/Framework/Models/SandboxSystemModel.h b/Engine/Plugins/Developer/Sandbox/SandboxedEditing/Source/SandboxedEditing/Private/Framework/Models/SandboxSystemModel.h
--- a/Engine/Plugins/Developer/Sandbox/SandboxedEditing/Source/SandboxedEditing/Private/Framework/Models/SandboxSystemModel.h
+++ b/Engine/Plugins/Developer/Sandbox/SandboxedEditing/Source/SandboxedEditing/Private/Framework/Models/SandboxSystemModel.h
@@ -13,6 +13,7 @@ class ISandboxInstance;
 enum class EBreakBehavior : uint8;
 enum class ESandboxFileChange : uint8;
 struct FGatheredFileChanges;
+struct FSandboxedFileChangeInfo;
 }
 
 template<typename OptionalType> struct TOptional;
@@ -83,11 +84,8 @@ public:
 	/** Invokes InCallback for each file change. */
 	void EnumerateFileChanges(
 		const FString& InSandbox, 
-		TFunctionRef<FileSandboxCore::EBreakBehavior(const FString& InPath, FileSandboxCore::ESandboxFileChange InChange)> InCallback
+		TFunctionRef<FileSandboxCore::EBreakBehavior(const FileSandboxCore::FSandboxedFileChangeInfo& InChange)> InCallback
 		);
-	
-	/** @return Timestamp of the sandboxed file action, or FDateTime::MinValue. */
-	FDateTime GetSandboxedFileTimestamp(const FString& InPath) const;
 
 	/** Discards all changes */
 	void DiscardAllChanges() const;
