Sandboxed Editing: Fix file change timestamps not working due to reversed function arguments.

#Synced-CL 49159615

diff --git a/Engine/Plugins/Developer/Sandbox/FileSandbox/Source/FileSandboxCore/Private/Utils/SandboxFileUtils.cpp b/Engine/Plugins/Developer/Sandbox/FileSandbox/Source/FileSandboxCore/Private/Utils/SandboxFileUtils.cpp
--- a/Engine/Plugins/Developer/Sandbox/FileSandbox/Source/FileSandboxCore/Private/Utils/SandboxFileUtils.cpp
+++ b/Engine/Plugins/Developer/Sandbox/FileSandbox/Source/FileSandboxCore/Private/Utils/SandboxFileUtils.cpp
@@ -72,8 +72,8 @@ void EnumerateMountPoints(TFunctionRef<EBreakBehavior(const FString& InAssetPath
 }
 
 TOptional<FDateTime> GetSandboxTimestamp(
-	const FString& InRootSandboxPath, 
-	const FString& InNonSandboxPath, const FFileSandboxCore_ManifestData& InManifest
+	const FString& InNonSandboxPath, 
+	const FString& InRootSandboxPath, const FFileSandboxCore_ManifestData& InManifest
 	)
 {
 	const TArray<FString>& RemovedFiles = InManifest.DeletedFiles.GetFiles();
