Change SEW to use FileSandbox versus the implementation in the AssetManagement plugin.

#changelist validated
#Synced-CL 49807761

diff --git a/Engine/Plugins/Developer/Sandbox/FileSandbox/Source/FileSandboxCore/Private/Sandbox/ISandboxInstance.cpp b/Engine/Plugins/Developer/Sandbox/FileSandbox/Source/FileSandboxCore/Private/Sandbox/ISandboxInstance.cpp
--- a/Engine/Plugins/Developer/Sandbox/FileSandbox/Source/FileSandboxCore/Private/Sandbox/ISandboxInstance.cpp
+++ b/Engine/Plugins/Developer/Sandbox/FileSandbox/Source/FileSandboxCore/Private/Sandbox/ISandboxInstance.cpp
@@ -11,11 +11,12 @@
 
 namespace UE::FileSandboxCore
 {
-void ISandboxInstance::PersistAll()
+bool ISandboxInstance::PersistAll()
 {
 	const FGatheredFileChanges FileChanges = GatherChangedFiles(EFileChangeGatherFlags::None);
 	const TArray<FString>& ChangedFiles = FileChanges.NonSandboxPaths;
-	PersistSandbox(FPersistArgs{ ChangedFiles });
+	FPersistResult Result = PersistSandbox(FPersistArgs{ ChangedFiles });
+	return Result.PersistStatus == EPersistStatus::Success;
 }
 
 void ISandboxInstance::DiscardAll()
diff --git a/Engine/Plugins/Developer/Sandbox/FileSandbox/Source/FileSandboxCore/Public/FileSandboxCoreMinimal.h b/Engine/Plugins/Developer/Sandbox/FileSandbox/Source/FileSandboxCore/Public/FileSandboxCoreMinimal.h
new file mode 100644
--- /dev/null
+++ b/Engine/Plugins/Developer/Sandbox/FileSandbox/Source/FileSandboxCore/Public/FileSandboxCoreMinimal.h
@@ -0,0 +1,12 @@
+// Copyright Epic Games, Inc. All Rights Reserved.
+#pragma once
+
+#include "IFileSandboxCoreModule.h"
+#include "ISandboxManager.h"
+#include "ISandboxInstance.h"
+
+#include "Types/ArgsFwds.h"
+#include "Types/Manager/NewSandboxArgs.h"
+#include "Types/Manager/LoadSandboxByNameArgs.h"
+#include "Types/Manager/DeleteSandboxResult.h"
+#include "Types/Manager/DeleteSandboxByNameArgs.h"
diff --git a/Engine/Plugins/Developer/Sandbox/FileSandbox/Source/FileSandboxCore/Public/ISandboxInstance.h b/Engine/Plugins/Developer/Sandbox/FileSandbox/Source/FileSandboxCore/Public/ISandboxInstance.h
--- a/Engine/Plugins/Developer/Sandbox/FileSandbox/Source/FileSandboxCore/Public/ISandboxInstance.h
+++ b/Engine/Plugins/Developer/Sandbox/FileSandbox/Source/FileSandboxCore/Public/ISandboxInstance.h
@@ -34,7 +34,7 @@ class ISandboxInstance
 public:
 	
 	/** Persists all files. */
-	UE_API void PersistAll();
+	UE_API bool PersistAll();
 	
 	/**
 	 * More advanced version of PersistAll that gives you more control how persisting should occur.
@@ -101,4 +101,4 @@ public:
 };
 }
 
-#undef UE_API
\ No newline at end of file
+#undef UE_API
