When iterating over the changed files, make sure we consider the original timestamp for deleted assets.

#rnx

#Synced-CL 49936699

diff --git a/Engine/Plugins/Developer/Sandbox/FileSandbox/Source/FileSandboxCore/Private/Sandbox/SandboxInstance.cpp b/Engine/Plugins/Developer/Sandbox/FileSandbox/Source/FileSandboxCore/Private/Sandbox/SandboxInstance.cpp
--- a/Engine/Plugins/Developer/Sandbox/FileSandbox/Source/FileSandboxCore/Private/Sandbox/SandboxInstance.cpp
+++ b/Engine/Plugins/Developer/Sandbox/FileSandbox/Source/FileSandboxCore/Private/Sandbox/SandboxInstance.cpp
@@ -32,13 +32,18 @@ static void TrackRevertedFileChangeInManifest(FFileSandboxCore_ManifestData& Out
 	}
 }
 
-static void TrackFileChangeInManifest(FFileSandboxCore_ManifestData& OutManifestData, const FString& InFile, ESandboxFileChange InAction)
+static void TrackFileChangeInManifest(FFileSandboxCore_ManifestData& OutManifestData, const FSandboxedFileChangeInfo& InChange)
 {
-	switch (InAction)
+	switch (InChange.Action)
+	{
+	case ESandboxFileChange::Added: OutManifestData.AddedFiles.Add(InChange.Path); break;
+	case ESandboxFileChange::Edited: OutManifestData.ModifiedFiles.Add(InChange.Path); break;
+	case ESandboxFileChange::Removed:
 	{
-	case ESandboxFileChange::Added: OutManifestData.AddedFiles.Add(InFile); break;
-	case ESandboxFileChange::Edited: OutManifestData.ModifiedFiles.Add(InFile); break;
-	case ESandboxFileChange::Removed: OutManifestData.DeletedFiles.Add(InFile); break;
+		FDateTime Timestamp = InChange.Timestamp != FDateTime::MinValue() ? InChange.Timestamp : FDateTime::UtcNow();
+		OutManifestData.DeletedFiles.Add(InChange.Path, Timestamp);
+		break;
+	}
 	default: break;
 	}
 }
@@ -254,9 +259,9 @@ void FSandboxInstance::OnEndOfFrame()
 		NewManifest.VersionInfo.Initialize();
 		EnumerateFileChanges([&NewManifest](const FSandboxedFileChangeInfo& Change)
 		{
-			Private::TrackFileChangeInManifest(NewManifest, Change.Path, Change.Action);
+			Private::TrackFileChangeInManifest(NewManifest, Change);
 			return EBreakBehavior::Continue;
-		}, EFileEnumerationFlags::None);
+		}, EFileEnumerationFlags::IncludeTimestamps);
 		
 		ManifestContent = MoveTemp(NewManifest);
 		SaveManifest();
