File Sandbox: Fix failing checkf when deleting sandbox.

The reason was that FDirectoryWatcherWindows::UnregisterDirectoryChangedCallback_Handle was being called while DirectoryWatcher was processing OS notifications.
Essentially, FDirectoryWatcherWindows::Tick called FWatchedSandboxRepository::UpdateWatchedDirectories, which called FDirectoryWatcherWindows::UnregisterDirectoryChangedCallback_Handle.
Both Tick and UnregisterDirectoryChangedCallback_Handle modify FDirectoryWatcherWindows::RequestMap, so Tick would fail a check because UnregisterDirectoryChangedCallback_Handle recursively removed an item while Tick was iterating through it.

This right fix would probably be to make FDirectoryWatcherWindows::UnregisterDirectoryChangedCallback_Handle handle this case but I'm not comfortable with this code, and I also do not know what about the other platforms, like Mac, Linux, etc.

#rnx

#Synced-CL 49936699

diff --git a/Engine/Plugins/Developer/Sandbox/FileSandbox/Source/FileSandboxCore/Private/Repository/WatchedSandboxRepository.cpp b/Engine/Plugins/Developer/Sandbox/FileSandbox/Source/FileSandboxCore/Private/Repository/WatchedSandboxRepository.cpp
--- a/Engine/Plugins/Developer/Sandbox/FileSandbox/Source/FileSandboxCore/Private/Repository/WatchedSandboxRepository.cpp
+++ b/Engine/Plugins/Developer/Sandbox/FileSandbox/Source/FileSandboxCore/Private/Repository/WatchedSandboxRepository.cpp
@@ -7,6 +7,7 @@
 #include "GenericPlatform/GenericPlatformFile.h"
 #include "HAL/PlatformFileManager.h"
 #include "IDirectoryWatcher.h"
+#include "Misc/CoreDelegates.h"
 #include "Misc/Paths.h"
 #include "Types/EBreakBehavior.h"
 #include "Types/RepositoryChangedEvent.h"
@@ -51,8 +52,15 @@ FWatchedSandboxRepository::FWatchedSandboxRepository(FString InBaseDirectory)
 {
 	RescanBaseDirectory();
 	UpdateWatchedDirectories();
+	
+	FCoreDelegates::OnEndFrame.AddRaw(this, &FWatchedSandboxRepository::OnEndOfFrame);
 }
 
+FWatchedSandboxRepository::~FWatchedSandboxRepository()
+{
+	FCoreDelegates::OnEndFrame.RemoveAll(this);
+}
+	
 void FWatchedSandboxRepository::ForEachSandbox(TFunctionRef<EBreakBehavior(const FString& InRootPath, const FSandboxMetaInfo& MetaData)> InProcessSandboxes)
 {
 	for (const FSandboxData& SandboxData : CachedSandboxData)
@@ -228,6 +236,11 @@ bool FWatchedSandboxRepository::DetectManifestFileChange(int32 InSandboxIndex, c
 	}
 }
 
+void FWatchedSandboxRepository::OnEndOfFrame()
+{
+	PendingDirectoryWatchersToUnregister.Reset();
+}
+
 void FWatchedSandboxRepository::UpdateCacheFromChangedBaseDirectory(const TArray<FFileChangeData>& InChanges)
 {
 	const bool bDoFullRescan = Algo::AnyOf(InChanges, [](const FFileChangeData& Change)
@@ -315,7 +328,17 @@ void FWatchedSandboxRepository::UpdateWatchedDirectories()
 			
 		return true;
 	});
+	
+	// UpdateWatchedDirectories may be called by DirectoryWatcher to process changes.
+	// It's not valid to call UnregisterDirectoryChangedCallback_Handle while the DirectoryWatcher module is processing notifications, so defer.
+	for (FScopedWatchedDirectory& WatchedDirectory : Old)
+	{
+		if (WatchedDirectory.IsValid())
+		{
+			PendingDirectoryWatchersToUnregister.Emplace(MoveTemp(WatchedDirectory));
+		}
+	}
 }
 }
 
-#endif
\ No newline at end of file
+#endif
diff --git a/Engine/Plugins/Developer/Sandbox/FileSandbox/Source/FileSandboxCore/Private/Repository/WatchedSandboxRepository.h b/Engine/Plugins/Developer/Sandbox/FileSandbox/Source/FileSandboxCore/Private/Repository/WatchedSandboxRepository.h
--- a/Engine/Plugins/Developer/Sandbox/FileSandbox/Source/FileSandboxCore/Private/Repository/WatchedSandboxRepository.h
+++ b/Engine/Plugins/Developer/Sandbox/FileSandbox/Source/FileSandboxCore/Private/Repository/WatchedSandboxRepository.h
@@ -17,9 +17,10 @@ namespace UE::FileSandboxCore
 class FWatchedSandboxRepository : public ISandboxRepository, public FNoncopyable
 {
 public:
-
+	
 	explicit FWatchedSandboxRepository(FString InBaseDirectory);
-
+	virtual ~FWatchedSandboxRepository();
+		
 	//~ Begin ISandboxRepository Interface
 	virtual void ForEachSandbox(TFunctionRef<EBreakBehavior(const FString& InRootPath, const FSandboxMetaInfo& MetaData)> InProcessSandboxes) override;
 	virtual int32 NumSandboxes() const override { return CachedSandboxData.Num(); }
@@ -51,6 +52,8 @@ private:
 	 * All directories in the base directory are always watched so we can detect when a new sandbox is created there. 
 	 */
 	TArray<FScopedWatchedDirectory> WatchedSubDirectories;
+	/** The directory that are supposed to be unregistered at end of frame. */
+	TArray<FScopedWatchedDirectory> PendingDirectoryWatchersToUnregister;
 	
 	/** Invoked when the known sandboxes changes (removed, discovered, etc.) */
 	FOnSandboxesChanged OnSandboxesChangedDelegate;
@@ -78,7 +81,9 @@ private:
 	bool DetectMetaDataFileChange(int32 InSandboxIndex, const TArray<FFileChangeData>& InChanges, const FString& InSandboxPath);
 	/** Checks whether a change has been reported for the manifest file. */
 	bool DetectManifestFileChange(int32 InSandboxIndex, const TArray<FFileChangeData>& InChanges, const FString& InSandboxPath);
+	
+	void OnEndOfFrame();
 };	
 }
 
-#endif
\ No newline at end of file
+#endif
diff --git a/Engine/Plugins/Developer/Sandbox/FileSandbox/Source/FileSandboxCore/Private/Utils/Watcher/ScopedWatchedDirectory.h b/Engine/Plugins/Developer/Sandbox/FileSandbox/Source/FileSandboxCore/Private/Utils/Watcher/ScopedWatchedDirectory.h
--- a/Engine/Plugins/Developer/Sandbox/FileSandbox/Source/FileSandboxCore/Private/Utils/Watcher/ScopedWatchedDirectory.h
+++ b/Engine/Plugins/Developer/Sandbox/FileSandbox/Source/FileSandboxCore/Private/Utils/Watcher/ScopedWatchedDirectory.h
@@ -21,6 +21,7 @@ public:
 	~FScopedWatchedDirectory();
 	
 	const FString& GetWatchedDirectory() const { return WatchedDirectory; }
+	bool IsValid() const { return Handle.IsValid(); }
 
 private:
 	
