File Sandbox: Fix not being able to edit description of creating a sandbox.

The issue was that the sandbox repository did not receive any directory notification, so the system though it did not exist.
The solution is to just listen for creation event emitted by the sandbox manager.

#rnx

#Synced-CL 49936699

diff --git a/Engine/Plugins/Developer/Sandbox/FileSandbox/Source/FileSandboxCore/Private/FileSandboxCoreModule.cpp b/Engine/Plugins/Developer/Sandbox/FileSandbox/Source/FileSandboxCore/Private/FileSandboxCoreModule.cpp
--- a/Engine/Plugins/Developer/Sandbox/FileSandbox/Source/FileSandboxCore/Private/FileSandboxCoreModule.cpp
+++ b/Engine/Plugins/Developer/Sandbox/FileSandbox/Source/FileSandboxCore/Private/FileSandboxCoreModule.cpp
@@ -18,7 +18,7 @@ void FFileSandboxCoreModule::StartupModule()
 	SandboxManager = MakeUnique<FSandboxManager>();
 	
 #if UE_SANDBOX_WITH_DIRECTORY_WATCHER
-	DefaultRepository = MakeUnique<FWatchedSandboxRepository>(GetBaseSandboxDirectory());
+	DefaultRepository = MakeUnique<FWatchedSandboxRepository>(GetBaseSandboxDirectory(), *SandboxManager);
 #else
 	DefaultRepository = MakeUnique<FNaiveSandboxRepository>(GetBaseSandboxDirectory(), *SandboxManager);
 #endif
diff --git a/Engine/Plugins/Developer/Sandbox/FileSandbox/Source/FileSandboxCore/Private/Repository/WatchedSandboxRepository.cpp b/Engine/Plugins/Developer/Sandbox/FileSandbox/Source/FileSandboxCore/Private/Repository/WatchedSandboxRepository.cpp
--- a/Engine/Plugins/Developer/Sandbox/FileSandbox/Source/FileSandboxCore/Private/Repository/WatchedSandboxRepository.cpp
+++ b/Engine/Plugins/Developer/Sandbox/FileSandbox/Source/FileSandboxCore/Private/Repository/WatchedSandboxRepository.cpp
@@ -7,6 +7,8 @@
 #include "GenericPlatform/GenericPlatformFile.h"
 #include "HAL/PlatformFileManager.h"
 #include "IDirectoryWatcher.h"
+#include "ISandboxInstance.h"
+#include "ISandboxManager.h"
 #include "Misc/CoreDelegates.h"
 #include "Misc/Paths.h"
 #include "Types/EBreakBehavior.h"
@@ -17,8 +19,9 @@
 
 namespace UE::FileSandboxCore
 {
-FWatchedSandboxRepository::FWatchedSandboxRepository(FString InBaseDirectory)
+FWatchedSandboxRepository::FWatchedSandboxRepository(FString InBaseDirectory, ISandboxManager& InSandboxManager)
 	: BaseDirectory(MoveTemp(InBaseDirectory))
+	, SandboxManager(InSandboxManager)
 	, WatchedBaseDirectory([this]()
 	{
 		IDirectoryWatcher* DirectoryWatcher = GetDirectoryWatcher();
@@ -53,11 +56,13 @@ FWatchedSandboxRepository::FWatchedSandboxRepository(FString InBaseDirectory)
 	RescanBaseDirectory();
 	UpdateWatchedDirectories();
 	
+	SandboxManager.OnPostSandboxStartup().AddRaw(this, &FWatchedSandboxRepository::OnSandboxStartup);
 	FCoreDelegates::OnEndFrame.AddRaw(this, &FWatchedSandboxRepository::OnEndOfFrame);
 }
 
 FWatchedSandboxRepository::~FWatchedSandboxRepository()
 {
+	SandboxManager.OnPostSandboxStartup().RemoveAll(this);
 	FCoreDelegates::OnEndFrame.RemoveAll(this);
 }
 	
@@ -339,6 +344,11 @@ void FWatchedSandboxRepository::UpdateWatchedDirectories()
 		}
 	}
 }
+
+void FWatchedSandboxRepository::OnSandboxStartup(ISandboxInstance& SandboxInstance)
+{
+	AddSandboxIfExists(SandboxInstance.GetRootDirectory());
+}
 }
 
 #endif
diff --git a/Engine/Plugins/Developer/Sandbox/FileSandbox/Source/FileSandboxCore/Private/Repository/WatchedSandboxRepository.h b/Engine/Plugins/Developer/Sandbox/FileSandbox/Source/FileSandboxCore/Private/Repository/WatchedSandboxRepository.h
--- a/Engine/Plugins/Developer/Sandbox/FileSandbox/Source/FileSandboxCore/Private/Repository/WatchedSandboxRepository.h
+++ b/Engine/Plugins/Developer/Sandbox/FileSandbox/Source/FileSandboxCore/Private/Repository/WatchedSandboxRepository.h
@@ -13,13 +13,16 @@ struct FFileChangeData;
 
 namespace UE::FileSandboxCore
 {
+class ISandboxInstance;
+class ISandboxManager;
+
 /** Simple directory watcher that detects when sandboxes are added, removed, and metadata file edited. */
 class FWatchedSandboxRepository : public ISandboxRepository, public FNoncopyable
 {
 public:
 	
-	explicit FWatchedSandboxRepository(FString InBaseDirectory);
-	virtual ~FWatchedSandboxRepository();
+	explicit FWatchedSandboxRepository(FString InBaseDirectory, ISandboxManager& InSandboxManager UE_LIFETIMEBOUND);
+	virtual ~FWatchedSandboxRepository() override;
 		
 	//~ Begin ISandboxRepository Interface
 	virtual void ForEachSandbox(TFunctionRef<EBreakBehavior(const FString& InRootPath, const FSandboxMetaInfo& MetaData)> InProcessSandboxes) override;
@@ -33,6 +36,9 @@ private:
 	
 	/** The base directory in which to look for sandboxes. Only looks at the top-level, i.e. not recursively. */
 	const FString BaseDirectory;
+ 	
+	/** Used to be notified when a sandbox directory is created since the sandbox watcher may not always fire immediately. */
+	ISandboxManager& SandboxManager;
 	
 	struct FSandboxData
 	{
@@ -82,6 +88,8 @@ private:
 	/** Checks whether a change has been reported for the manifest file. */
 	bool DetectManifestFileChange(int32 InSandboxIndex, const TArray<FFileChangeData>& InChanges, const FString& InSandboxPath);
 	
+	/** Handles a sandbox being created. */
+	void OnSandboxStartup(ISandboxInstance& SandboxInstance);
 	void OnEndOfFrame();
 };	
 }
