Fix FileSandbox incorrectly tracking files added and then deleted within the same sandbox session

When a user added a new asset to a sandbox and then deleted it, the file would incorrectly appear in the sandbox's changed files list as "deleted" even though the net effect should be zero (the file
never existed outside the sandbox).

1. SandboxInstance.cpp: Modified the deletion tracking logic to check whether the file being deleted was previously added in the sandbox:
- Capture the result of removing from AddedFiles
- Only add to DeletedFiles if the file was NOT in AddedFiles (meaning it existed before sandbox)
- If it WAS in AddedFiles, the removal is sufficient (net-zero operation)
- Simplified refresh flag handling to consistently use UpdateManifest | BroadcastChanges
2. SandboxPlatformFile.cpp: Added missing BroadcastFileChange() call in DeleteFile() to ensure file deletion state changes are properly broadcast to listeners, maintaining consistency with other file
operations.

#rnx
#Synced-CL 49936699

diff --git a/Engine/Plugins/Developer/Sandbox/FileSandbox/Source/FileSandboxCore/Private/Sandbox/Platform/SandboxPlatformFile.cpp b/Engine/Plugins/Developer/Sandbox/FileSandbox/Source/FileSandboxCore/Private/Sandbox/Platform/SandboxPlatformFile.cpp
--- a/Engine/Plugins/Developer/Sandbox/FileSandbox/Source/FileSandboxCore/Private/Sandbox/Platform/SandboxPlatformFile.cpp
+++ b/Engine/Plugins/Developer/Sandbox/FileSandbox/Source/FileSandboxCore/Private/Sandbox/Platform/SandboxPlatformFile.cpp
@@ -209,6 +209,12 @@ bool FSandboxPlatformFile::DeleteFile(const TCHAR* Filename)
 		if (bPerformedDeleteOp)
 		{
 			NotifyFileDeleted(ResolvedPath);
+
+			const bool bWasEdit = bIsDeletingNonSandboxFile && bIsDeletingFileCreatedInSandbox;
+			const bool bWasAdd = !bIsDeletingNonSandboxFile && bIsDeletingFileCreatedInSandbox;
+
+			const ESandboxFileChange OldState = bWasEdit ? ESandboxFileChange::Edited : (bWasAdd ? ESandboxFileChange::Added : ESandboxFileChange::None);
+			BroadcastFileChange(ResolvedPath, OldState, ESandboxFileChange::Removed);
 		}
 		
 		return bPerformedDeleteOp; 
@@ -1129,4 +1135,4 @@ void FSandboxPlatformFile::BroadcastFileChange(const FSandboxedPlatformFilePath&
 	ensure(InOldState != InNewState);
 	OnFileStateChangedDelegate.Broadcast(InPath, InOldState, InNewState);
 }
-}
\ No newline at end of file
+}
diff --git a/Engine/Plugins/Developer/Sandbox/FileSandbox/Source/FileSandboxCore/Private/Sandbox/SandboxInstance.cpp b/Engine/Plugins/Developer/Sandbox/FileSandbox/Source/FileSandboxCore/Private/Sandbox/SandboxInstance.cpp
--- a/Engine/Plugins/Developer/Sandbox/FileSandbox/Source/FileSandboxCore/Private/Sandbox/SandboxInstance.cpp
+++ b/Engine/Plugins/Developer/Sandbox/FileSandbox/Source/FileSandboxCore/Private/Sandbox/SandboxInstance.cpp
@@ -209,32 +209,23 @@ bool FSandboxInstance::DeletedPackageExistsInNonSandbox(const FString& InFilenam
 void FSandboxInstance::OnFileDeletionStateChanged(const FSandboxedPlatformFilePath& InFile, ESandboxFileChange InOldState, ESandboxFileChange InNewState)
 {
 	const FNonSandboxPath NonSandboxPath = InFile.GetNonSandboxPath();
-	
+
+	const ESandboxRefreshFlags CurrentFlags = RefreshFlags.load();
 	const bool bWasDeleted = InNewState == ESandboxFileChange::Removed;
 	if (bWasDeleted)
 	{
 		ManifestContent.ModifiedFiles.Remove(NonSandboxPath);
-		ManifestContent.DeletedFiles.Add(NonSandboxPath);
-		ManifestContent.AddedFiles.Remove(NonSandboxPath);
+		const bool bWasAdded = ManifestContent.AddedFiles.Remove(NonSandboxPath);
+		if (!bWasAdded)
+		{
+			ManifestContent.DeletedFiles.Add(NonSandboxPath);
+		}
 		SaveManifest();
-		
-		// The manifest was updated already. Only need to broadcast at end of tick.
-		const ESandboxRefreshFlags CurrentFlags = RefreshFlags.load();
 		RefreshFlags = CurrentFlags | ESandboxRefreshFlags::BroadcastChanges;
+		return;
 	}
-	
-	const bool bClearedDeleted = InOldState == ESandboxFileChange::Removed;
-	if (bClearedDeleted)
-	{
-		ManifestContent.DeletedFiles.Remove(NonSandboxPath);
-		// At the end of tick, check whether the file was marked for add or edit
-		RefreshFlags = ESandboxRefreshFlags::BroadcastChanges | ESandboxRefreshFlags::UpdateManifest;
-	}
-	else
-	{
-		// At end of tick, rebuild the manifest based on the files in the sandbox directory.
-		RefreshFlags = ESandboxRefreshFlags::BroadcastChanges | ESandboxRefreshFlags::UpdateManifest;
-	}
+
+	RefreshFlags = CurrentFlags | ESandboxRefreshFlags::UpdateManifest | ESandboxRefreshFlags::BroadcastChanges;
 }
 
 bool FSandboxInstance::SaveManifest() const
