Fix edge case with edits of assets in a sandbox.  They were not getting captured and listed as a delete then an edit.

#rnx

#Synced-CL 49936699

diff --git a/Engine/Plugins/Developer/Sandbox/FileSandbox/Source/FileSandboxCore/Private/Sandbox/SandboxInstance.cpp b/Engine/Plugins/Developer/Sandbox/FileSandbox/Source/FileSandboxCore/Private/Sandbox/SandboxInstance.cpp
--- a/Engine/Plugins/Developer/Sandbox/FileSandbox/Source/FileSandboxCore/Private/Sandbox/SandboxInstance.cpp
+++ b/Engine/Plugins/Developer/Sandbox/FileSandbox/Source/FileSandboxCore/Private/Sandbox/SandboxInstance.cpp
@@ -7,6 +7,7 @@
 #include "Misc/PackageName.h"
 #include "Persist/DiffUtils.h"
 #include "Persist/LogPersistFeedback.h"
+#include "Platform/SandboxPlatformFile.h"
 #include "Types/Manager/SandboxInitArgs.h"
 #include "Types/Sandbox/DiscardResult.h"
 #include "Types/SandboxFileChange.h"
@@ -224,7 +225,16 @@ void FSandboxInstance::OnFileDeletionStateChanged(const FSandboxedPlatformFilePa
 		RefreshFlags = CurrentFlags | ESandboxRefreshFlags::BroadcastChanges;
 		return;
 	}
-
+	else if (InOldState == ESandboxFileChange::Removed
+  		  && InNewState == ESandboxFileChange::Edited) 
+	{
+		ManifestContent.DeletedFiles.Remove(NonSandboxPath);
+		ManifestContent.ModifiedFiles.Add(NonSandboxPath);
+		SaveManifest();
+		RefreshFlags = CurrentFlags | ESandboxRefreshFlags::BroadcastChanges;
+		return;
+	}
+	
 	RefreshFlags = CurrentFlags | ESandboxRefreshFlags::UpdateManifest | ESandboxRefreshFlags::BroadcastChanges;
 }
 
