File Sandbox Core: Add API for importing and exporting sandbox as zip.

#Synced-CL 50904440

diff --git a/Engine/Plugins/Developer/Sandbox/FileSandbox/Source/FileSandboxCore/FileSandboxCore.Build.cs b/Engine/Plugins/Developer/Sandbox/FileSandbox/Source/FileSandboxCore/FileSandboxCore.Build.cs
--- a/Engine/Plugins/Developer/Sandbox/FileSandbox/Source/FileSandboxCore/FileSandboxCore.Build.cs
+++ b/Engine/Plugins/Developer/Sandbox/FileSandbox/Source/FileSandboxCore/FileSandboxCore.Build.cs
@@ -17,6 +17,7 @@ namespace UnrealBuildTool.Rules
 				new string[]
 				{
 					"Engine",
+					"FileUtilities",
 					"Json",
 					"JsonUtilities", 
 					"RenderCore",
diff --git a/Engine/Plugins/Developer/Sandbox/FileSandbox/Source/FileSandboxCore/Private/Utils/SandboxDirectoryUtils.cpp b/Engine/Plugins/Developer/Sandbox/FileSandbox/Source/FileSandboxCore/Private/Utils/SandboxDirectoryUtils.cpp
--- a/Engine/Plugins/Developer/Sandbox/FileSandbox/Source/FileSandboxCore/Private/Utils/SandboxDirectoryUtils.cpp
+++ b/Engine/Plugins/Developer/Sandbox/FileSandbox/Source/FileSandboxCore/Private/Utils/SandboxDirectoryUtils.cpp
@@ -32,7 +32,12 @@ FString GetBaseSandboxDirectory()
 		? FPaths::ProjectIntermediateDir() / TEXT("Sandboxes")
 		: CommandLineOverride;
 }
-	
+
+bool ValidateManifestContent(const FString& InFileContent)
+{
+	return LoadManifestFromContent(InFileContent).IsSet();
+}
+
 void ForEachSandbox(
 	TFunctionRef<FProcessSandboxDirectorySignature> InProcessSandbox,
 	const FString& InBaseDirectory
@@ -52,9 +57,10 @@ void ForEachSandbox(
 
 bool IsRootSandboxDirectory(const FString& InRootDirectory)
 {
-	return FPaths::DirectoryExists(InRootDirectory)
-		&& FPaths::FileExists(InRootDirectory / GetManifestFileName())
-		&& FPaths::FileExists(InRootDirectory / GetMetadataFileName());
+	const bool bDirectoryExists = FPaths::DirectoryExists(InRootDirectory);
+	const bool bHasManifest = FPaths::FileExists(InRootDirectory / GetManifestFileName());
+	const bool bHasMetadata = FPaths::FileExists(InRootDirectory / GetMetadataFileName());
+	return bDirectoryExists && bHasManifest && bHasMetadata;
 }
 
 bool SaveMetaData(
@@ -75,19 +81,25 @@ bool SaveMetaData(
 	
 TOptional<FFileSandboxCore_SandboxMetaData> LoadMetaData(const FString& InSandboxDirectory)
 {
-	FFileSandboxCore_SandboxMetaData ManifestPlatformFileState;
 	FString JSONPayload;
-	
 	const FString AbsFileName = InSandboxDirectory / GetMetadataFileName();
 	if (FFileHelper::LoadFileToString(JSONPayload, *AbsFileName))
 	{
-		const bool bParsed = FJsonObjectConverter::JsonObjectStringToUStruct(JSONPayload, &ManifestPlatformFileState, 0, 0);
-		UE_CLOG(!bParsed, LogFileSandbox, Error, TEXT("Failed to parse sandbox metadata file %s"), *AbsFileName);
-		return bParsed ? ManifestPlatformFileState : TOptional<FFileSandboxCore_SandboxMetaData>{};
+		const TOptional<FFileSandboxCore_SandboxMetaData> Result = LoadMetaDataFromFileContent(JSONPayload);
+		UE_CLOG(!Result.IsSet(), LogFileSandbox, Error, TEXT("Failed to parse sandbox metadata file %s"), *AbsFileName);
+		return Result;
 	}
 	return {};
 }
 
+TOptional<FFileSandboxCore_SandboxMetaData> LoadMetaDataFromFileContent(const FString& InFileContent)
+{
+	FFileSandboxCore_SandboxMetaData ManifestPlatformFileState;
+	const FString& JSONPayload = InFileContent;
+	const bool bParsed = FJsonObjectConverter::JsonObjectStringToUStruct(JSONPayload, &ManifestPlatformFileState, 0, 0);
+	return bParsed ? ManifestPlatformFileState : TOptional<FFileSandboxCore_SandboxMetaData>{};
+}
+
 bool RenameSandboxWithDirectory(
 	const FString& InNewName, const FString& InOldName, const FString& InBaseDirectory, const FFileSandboxCore_SandboxMetaData* InNewMetaData
 	)
@@ -295,6 +307,18 @@ TOptional<FString> GetSandboxPathFor(const FString& InSandboxDirectory, const FS
 	
 	return Result;
 }
+
+TOptional<FString> GetSandboxName(const FString& InSandboxDirectory)
+{
+	if (!IsRootSandboxDirectory(InSandboxDirectory))
+	{
+		return {};
+	}
+	
+	// We'll assume that the directory has the name, i.e. has a structure like ../Sandboxes/MySandboxName
+	const FString& PathLeaf = FPaths::GetPathLeaf(InSandboxDirectory);
+	return ensure(!PathLeaf.IsEmpty()) ? PathLeaf : TOptional<FString>();
+}
 }
 
 #undef LOCTEXT_NAMESPACE
\ No newline at end of file
diff --git a/Engine/Plugins/Developer/Sandbox/FileSandbox/Source/FileSandboxCore/Private/Utils/SandboxFileUtils.cpp b/Engine/Plugins/Developer/Sandbox/FileSandbox/Source/FileSandboxCore/Private/Utils/SandboxFileUtils.cpp
--- a/Engine/Plugins/Developer/Sandbox/FileSandbox/Source/FileSandboxCore/Private/Utils/SandboxFileUtils.cpp
+++ b/Engine/Plugins/Developer/Sandbox/FileSandbox/Source/FileSandboxCore/Private/Utils/SandboxFileUtils.cpp
@@ -2,9 +2,9 @@
 
 #include "SandboxFileUtils.h"
 
+#include "Data/ManifestData.h"
 #include "JsonObjectConverter.h"
 #include "LogFileSandbox.h"
-#include "Data/ManifestData.h"
 #include "Misc/FileHelper.h"
 #include "Misc/Optional.h"
 #include "Misc/PackageName.h"
@@ -12,6 +12,7 @@
 #include "Sandbox/Platform/SanboxedFilePathUtils.h"
 #include "Sandbox/Platform/SandboxedPlatformFilePath.h"
 #include "Types/EBreakBehavior.h"
+#include "Utils/SandboxDirectoryUtils.h"
 
 namespace UE::FileSandboxCore
 {
@@ -33,15 +34,23 @@ TOptional<FFileSandboxCore_ManifestData> LoadManifest(const FString& InRootDirec
 	const FString AbsFileName = InRootDirectory / GetManifestFileName();
 	if (FFileHelper::LoadFileToString(JSONPayload, *AbsFileName))
 	{
-		const bool bSuccess = FJsonObjectConverter::JsonObjectStringToUStruct(JSONPayload, &ManifestPlatformFileState, 0, 0);
-		UE_CLOG(!bSuccess, LogFileSandbox, Error, TEXT("Failed to parse sandbox manifest file %s"), *InRootDirectory);
-		return bSuccess ? ManifestPlatformFileState : TOptional<FFileSandboxCore_ManifestData>{};
+		const TOptional<FFileSandboxCore_ManifestData> Result = LoadManifestFromContent(JSONPayload);
+		UE_CLOG(!Result.IsSet(), LogFileSandbox, Error, TEXT("Failed to parse sandbox manifest file %s"), *InRootDirectory);
+		return Result;
 	}
 	
 	UE_CLOG(bLogErrorIfNotExist, LogFileSandbox, Error, TEXT("Failed to load sandbox manifest file %s"), *InRootDirectory);
 	return {};
 }
 
+TOptional<FFileSandboxCore_ManifestData> LoadManifestFromContent(const FString& InFileContent)
+{
+	FFileSandboxCore_ManifestData ManifestPlatformFileState;
+	const FString& JSONPayload = InFileContent;
+	const bool bSuccess = FJsonObjectConverter::JsonObjectStringToUStruct(JSONPayload, &ManifestPlatformFileState, 0, 0);
+	return bSuccess ? ManifestPlatformFileState : TOptional<FFileSandboxCore_ManifestData>{};
+}
+
 FString GetManifestFileName()
 {
 	return TEXT("manifest.json");
diff --git a/Engine/Plugins/Developer/Sandbox/FileSandbox/Source/FileSandboxCore/Private/Utils/SandboxFileUtils.h b/Engine/Plugins/Developer/Sandbox/FileSandbox/Source/FileSandboxCore/Private/Utils/SandboxFileUtils.h
--- a/Engine/Plugins/Developer/Sandbox/FileSandbox/Source/FileSandboxCore/Private/Utils/SandboxFileUtils.h
+++ b/Engine/Plugins/Developer/Sandbox/FileSandbox/Source/FileSandboxCore/Private/Utils/SandboxFileUtils.h
@@ -9,6 +9,7 @@ class FString;
 enum class EBreakBehavior : uint8;
 struct FDateTime;
 struct FFileSandboxCore_ManifestData;
+struct FFileSandboxCore_SandboxMetaData;
 template<typename OptionalType> struct TOptional;
 
 namespace UE::FileSandboxCore
@@ -17,6 +18,11 @@ namespace UE::FileSandboxCore
 bool SaveManifest(const FFileSandboxCore_ManifestData& InData, const FString& InRootDirectory);
 /** Loads the manifest file. */
 TOptional<FFileSandboxCore_ManifestData> LoadManifest(const FString& InRootDirectory, bool bLogErrorIfNotExist = true);
+/** Loads the manifest if you already have the file content. */
+TOptional<FFileSandboxCore_ManifestData> LoadManifestFromContent(const FString& InFileContent);
+
+/** Reads metadata about a sandbox if you already have the content of the file. */
+TOptional<FFileSandboxCore_SandboxMetaData> LoadMetaDataFromFileContent(const FString& InFileContent);
 
 /* @return The name to give the manifest file. */
 FString GetManifestFileName();
diff --git a/Engine/Plugins/Developer/Sandbox/FileSandbox/Source/FileSandboxCore/Private/Utils/ZipTransferUtils.cpp b/Engine/Plugins/Developer/Sandbox/FileSandbox/Source/FileSandboxCore/Private/Utils/ZipTransferUtils.cpp
new file mode 100644
--- /dev/null
+++ b/Engine/Plugins/Developer/Sandbox/FileSandbox/Source/FileSandboxCore/Private/Utils/ZipTransferUtils.cpp
@@ -0,0 +1,374 @@
+﻿// Copyright Epic Games, Inc. All Rights Reserved.
+
+#include "Utils/ZipTransferUtils.h"
+
+#include "Algo/AllOf.h"
+#include "Data/ManifestData.h"
+#include "FileUtilities/ZipArchiveReader.h"
+#include "FileUtilities/ZipArchiveWriter.h"
+#include "HAL/FileManager.h"
+#include "HAL/PlatformFileManager.h"
+#include "LogFileSandbox.h"
+#include "Logging/LogMacros.h"
+#include "Misc/FileHelper.h"
+#include "Misc/Paths.h"
+#include "SandboxFileUtils.h"
+#include "FileChange/EFileChangeAction.h"
+#include "Utils/SandboxDirectoryUtils.h"
+
+#if UE_SANDBOX_WITH_DIRECTORY_WATCHER
+#include "DirectoryWatcherModule.h"
+#include "Watcher/DirectoryWatcherUtils.h"
+#endif
+
+namespace UE::FileSandboxCore::ZipUtils
+{
+namespace ZipDetail
+{
+static bool WriteToZip(const FString& InSandboxRoot, IFileHandle& InZipFile, TConstArrayView<FString> InFilesToCompress)
+{
+	IFileManager& FileManager = IFileManager::Get();
+	// FZipArchiveWriter takes ownership of InZipFile. It handles releasing the resources.
+	FZipArchiveWriter ZipArchiveWriter(&InZipFile, EZipArchiveOptions::RemoveDuplicate | EZipArchiveOptions::Deflate);
+	
+	for (const FString& FileToCompress : InFilesToCompress)
+	{
+		const TUniquePtr<FArchive> Reader = TUniquePtr<FArchive>(FileManager.CreateFileReader(*FileToCompress));
+		if (!Reader)
+		{
+			UE_LOG(LogFileSandbox, Error, TEXT("Failed to open %s for reading"), *FileToCompress)
+			return false;
+		}
+		
+		const int64 Size = Reader->TotalSize();
+		TArray64<uint8> RawData;
+		RawData.SetNumUninitialized(Size);
+		Reader->Serialize(RawData.GetData(), Size);
+		Reader->Close();
+		
+		// FileToCompress is an absolute path, e.g. "D:/MyProject/Intermediate/Sandboxes/MySandbox/manifest.json".
+		// This converts it to "MySandbox/manifest.json".
+		FString RelativeToRoot = FileToCompress;
+		FPaths::MakePathRelativeTo(RelativeToRoot, *InSandboxRoot);
+		ZipArchiveWriter.AddFile(RelativeToRoot, RawData, FDateTime::Now());
+	}
+	
+	return true;
+}
+}
+
+bool ExportSandboxToZip(const FString& InSandboxRoot, const FString& InOutputZipPath)
+{
+	if (!ensure(FileSandboxCore::IsRootSandboxDirectory(InSandboxRoot)))
+	{
+		return false;
+	}
+	
+	IPlatformFile& PlatformFile = FPlatformFileManager::Get().GetPlatformFile();
+	
+	// Find all the files in the directory
+	TArray<FString> FilesToCompress;
+	PlatformFile.FindFilesRecursively(FilesToCompress, *InSandboxRoot, TEXT(""));
+	if (FilesToCompress.IsEmpty())
+	{
+		return false;
+	}
+
+	// Make sure that the output is .zip, in case the caller forgot it (e.g. they passed "D:/MyFoo/ZipFile" instead of "D:/MyFoo/ZipFile.zip"
+	const FString CleansedOutput = FPaths::SetExtension(InOutputZipPath, TEXT("zip"));
+	
+	// Creates a zip file
+	const bool bZipExisted = PlatformFile.FileExists(*CleansedOutput);
+	IFileHandle* ZipFile = PlatformFile.OpenWrite(*CleansedOutput); 
+	if (!ZipFile)
+	{
+		return false;
+	}
+	
+	// WriteToZip takes ownership of the IFileHandle handle.
+	if (!ZipDetail::WriteToZip(InSandboxRoot, *ZipFile, FilesToCompress))
+	{
+		// Unfortunately, FZipArchiveWriter API does not provide any API for cancelling. ~FZipArchiveWriter will just write, so we have to undo. 
+		// It could be that before the operation, CleansedOutput already contained a file which is being overwritten.
+		// In that case, this DeleteFile is an unfortunate casualty, but we'll assume that the user has already confirmed overwriting.
+		UE_CLOG(bZipExisted, LogFileSandbox, Warning, 
+			TEXT("Failed to export sandbox. The Zip %s already existed before. Clearing invalid data."), *CleansedOutput
+			)
+		PlatformFile.DeleteFile(*CleansedOutput);
+		return false;
+	}
+	
+	return true;
+}
+
+namespace ImportDetail
+{
+/** 
+ * @return First directory in the path, e.g. "Foo" in "/Foo/Bar/", "Foo/Bar", or "/Foo/Bar.temp". 
+ * Returns empty if no directory is contained, e.g. "/Foo". 
+ */
+static TOptional<FString> GetFirstDirectory(const FString& InPath)
+{
+	FString Normalized = InPath;
+	FPaths::NormalizeFilename(Normalized);
+
+	// Skip leading slash
+	if (Normalized.StartsWith(TEXT("/")))
+	{
+		Normalized = Normalized.RightChop(1);
+	}
+
+	int32 SlashIndex;
+	if (Normalized.FindChar(TEXT('/'), SlashIndex))
+	{
+		return Normalized.Left(SlashIndex);
+	}
+
+	return Normalized;
+}
+
+/** @return The path without the first directory, e.g. "Bar.txt" in "/Foo/Bar.txt". */
+static TOptional<FString> RemoveFirstDirectory(const FString& InPath)
+{
+	FString Normalized = InPath;
+	FPaths::NormalizeFilename(Normalized);
+
+	// Skip leading slash
+	if (Normalized.StartsWith(TEXT("/")))
+	{
+		Normalized = Normalized.RightChop(1);
+	}
+
+	int32 SlashIndex;
+	if (Normalized.FindChar(TEXT('/'), SlashIndex))
+	{
+		// E.g. "Foo/MyText.text" -> SlashIndex = 3, Len = 15,
+		// 15 - 3 = 12 -> Right(12) -> "/MyText.text"
+		// 15 - 3 - 1 = 11 -> Right(11) -> "MyText.text"
+		return Normalized.Right(Normalized.Len() - SlashIndex - 1);
+	}
+
+	// There is no first directory.
+	return {};
+}
+
+static TOptional<FString> GetDirectorySharedByAllFiles(const TArray<FString>& InAllFiles)
+{
+	const TOptional<FString> SandboxName = InAllFiles.IsEmpty() ? TOptional<FString>() : ImportDetail::GetFirstDirectory(InAllFiles[0]);
+	if (!SandboxName)
+	{
+		return SandboxName;
+	}
+	
+	const bool bAllFilesInSameDirectory = Algo::AllOf(InAllFiles, [&SandboxName](const FString& InFile)
+	{
+		const TOptional<FString> FileFirstDirectory = ImportDetail::GetFirstDirectory(InFile);
+		return FileFirstDirectory && *FileFirstDirectory == *SandboxName;
+	});
+	return bAllFilesInSameDirectory ? SandboxName : TOptional<FString>();
+}
+
+static bool CanParseRequiredFiles(const FString& InSandboxName, FImportInspectionResult& InOutResult)
+{
+	FZipArchiveReader& FileReader = *InOutResult.Reader;
+
+	const FString ManifestFilePath = FPaths::Combine(*InSandboxName, FileSandboxCore::GetManifestFileName());
+	const FString MetadataFilePath = FPaths::Combine(*InSandboxName, FileSandboxCore::GetMetadataFileName());
+	
+	TArray<uint8>& ManifestContentBytes = InOutResult.ManifestBytes;
+	TArray<uint8>& MetadataContentBytes = InOutResult.MetadataBytes;
+	const bool bHasManifest = FileReader.TryReadFile(ManifestFilePath, ManifestContentBytes);
+	const bool bHasMetadata = FileReader.TryReadFile(MetadataFilePath, MetadataContentBytes);
+	if (!bHasManifest || !bHasMetadata)
+	{
+		return false;
+	}
+	
+	FString ManifestContent, MetadataContent;
+	FFileHelper::BufferToString(ManifestContent, ManifestContentBytes.GetData(), ManifestContentBytes.Num());
+	FFileHelper::BufferToString(MetadataContent, MetadataContentBytes.GetData(), MetadataContentBytes.Num());
+	
+	InOutResult.MetaData = LoadMetaDataFromFileContent(MetadataContent);
+	const bool bCanParseRequiredFiles = InOutResult.MetaData.IsSet() && LoadManifestFromContent(ManifestContent).IsSet();
+	return bCanParseRequiredFiles;
+}
+}
+
+FImportInspectionResult::FImportInspectionResult(TUniquePtr<FZipArchiveReader> InReader)
+	: Reader(MoveTemp(InReader))
+	, AllFileNames(Reader->GetFileNames())
+{}
+
+TOptional<FImportInspectionResult> InspectFileForImport(const FString& InPathToZip)
+{
+	IPlatformFile& PlatformFile = FPlatformFileManager::Get().GetPlatformFile();
+	IFileHandle* ZipFile = PlatformFile.OpenRead(*InPathToZip); 
+	if (!ZipFile)
+	{
+		return {};
+	}
+	
+	FImportInspectionResult Result(MakeUnique<FZipArchiveReader>(ZipFile));
+	Result.AllFileNames = Result.Reader->GetFileNames();
+	
+	const TOptional<FString> SandboxName = ImportDetail::GetDirectorySharedByAllFiles(Result.AllFileNames);
+	if (!SandboxName || !ImportDetail::CanParseRequiredFiles(*SandboxName, Result))
+	{
+		return Result;
+	}
+	
+	// Setting SandboxName indicates that the sandbox is valid.
+	Result.SandboxName = *SandboxName;
+	return Result;
+}
+
+
+namespace ImportDetail
+{
+static bool ImportZippedFile(const TArray<uint8>& InFileContent, const FString& InTargetFile)
+{
+	IPlatformFile& PlatformFile = FPlatformFileManager::Get().GetPlatformFile();
+	const FString Directory = FPaths::GetPath(InTargetFile);
+	if (!PlatformFile.CreateDirectoryTree(*Directory))
+	{
+		return false;
+	}
+	
+	IFileManager& FileManager = IFileManager::Get();
+	const TUniquePtr<FArchive> Archive(FileManager.CreateFileWriter(*InTargetFile, 0));
+	if (!Archive)
+	{
+		return false;
+	}
+	
+	Archive->Serialize(const_cast<uint8*>(InFileContent.GetData()), InFileContent.Num());
+	Archive->Close(); // Always explicitly close to catch errors from flush/close
+	const bool bIsSuccess = !Archive->IsError() && !Archive->IsCriticalError();
+	
+	if (bIsSuccess)
+	{
+		// TODO UE-366795: Set the timestamp of created file to be equal to timestamp the file has in the zip archive. 
+		// Currently, FZipArchiveReader has no API to get that information.
+	}
+	
+	return bIsSuccess;
+}
+
+static void DeleteFiles(const TArray<FString>& InFiles)
+{
+	IPlatformFile& PlatformFile = FPlatformFileManager::Get().GetPlatformFile();
+	for (const FString& File : InFiles)
+	{
+		PlatformFile.DeleteFile(*File);
+	}
+}
+
+static void NotifyDirectoryWatcher(const TArray<FString>& InFiles)
+{
+#if UE_SANDBOX_WITH_DIRECTORY_WATCHER
+	// Notify that the sandboxed directories have been restored to their original state
+	FDirectoryWatcherModule* DirectoryWatcherModule = GetDirectoryWatcherModuleIfLoaded();
+	if (!DirectoryWatcherModule)
+	{
+		return;
+	}
+	
+	TArray<FFileChangeData> FileChanges;
+	Algo::Transform(InFiles, FileChanges, [](const FString& InFile)
+	{
+		return FFileChangeData(InFile, FFileChangeData::FCA_Added);
+	});
+	DirectoryWatcherModule->RegisterExternalChanges(FileChanges);
+#endif
+}
+
+template<typename TGetCachedFileContent> requires std::is_invocable_r_v<const TArray<uint8>*, TGetCachedFileContent, const FString& /*InFileName*/>
+static bool ImportFromZip(
+	const FZipArchiveReader& InReader, const TArray<FString>& InFilesInZipArchive, const FString& InSandboxTargetRoot, 
+	TGetCachedFileContent&& InGetCachedFileContent
+	)
+{
+	// 1. Nothing can be in the directory that we'll import to.
+	IPlatformFile& PlatformFile = FPlatformFileManager::Get().GetPlatformFile();
+	PlatformFile.DeleteDirectoryRecursively(*InSandboxTargetRoot);
+	
+	// 2. Copy all the files from the zip
+	TArray<FString> FilesImportedSoFar;
+	TArray<uint8> FileContent;
+	bool bFreeOfErrors = true;
+	for (int32 Index = 0; Index < InFilesInZipArchive.Num() && bFreeOfErrors; ++Index)
+	{
+		const FString& ZippedFile = InFilesInZipArchive[Index];
+		
+		// All paths are "SandboxName/Some/Foo.bar" but we want only want "Some/Foo.bar" to combine it with InSandboxTargetRoot.
+		const TOptional<FString> PathInSandbox = RemoveFirstDirectory(ZippedFile);
+		if (!ensure(PathInSandbox))
+		{
+			bFreeOfErrors = false;
+			break;
+		}
+		
+		// E.g. InSandboxTargetRoot = "D:Project/Intermediate/Sandboxes/", PathInSandbox = "Sandbox/Game/Material.uasset".
+		const FString Destination = FPaths::Combine(InSandboxTargetRoot, *PathInSandbox);
+		// This is for security to avoid "Zip Slips": the zip could contain "../../Outside.txt", so it would write outside of the intended directory.
+		// Code calling ImportFromZip should have already validated this, but we'll do it once more here to make sure.
+		if (!ensure(FPaths::IsUnderDirectory(Destination, InSandboxTargetRoot)))
+		{
+			bFreeOfErrors = false;
+			break;
+		}
+		
+		// Optimization: Is it a file that we already deserialized during inspection?
+		if (const TArray<uint8>* CachedContent = InGetCachedFileContent(ZippedFile))
+		{
+			bFreeOfErrors &= ImportZippedFile(*CachedContent, Destination);
+			FilesImportedSoFar.Add(Destination);
+			continue;
+		}
+		
+		FileContent.Empty(FileContent.Num());
+		if (InReader.TryReadFile(ZippedFile, FileContent))
+		{
+			bFreeOfErrors &= ImportZippedFile(FileContent, Destination);
+			FilesImportedSoFar.Add(Destination);
+		}
+		else
+		{
+			bFreeOfErrors = false;
+		}
+	}
+	
+	// 3. Potentially rollback any changes we've made if there were errors.
+	if (bFreeOfErrors)
+	{
+		// This is needed so FWatchedSandboxRepository is notified.
+		NotifyDirectoryWatcher(FilesImportedSoFar);
+	}
+	else
+	{
+		DeleteFiles(FilesImportedSoFar);
+	}
+	
+	return bFreeOfErrors;
+}
+}
+
+bool ImportSandboxFromZip(const FImportInspectionResult& InInspection, const FString& InSandboxTargetRoot)
+{
+	const FString ManifestFilePath = FPaths::Combine(InInspection.SandboxName, FileSandboxCore::GetManifestFileName());
+	const FString MetadataFilePath = FPaths::Combine(InInspection.SandboxName, FileSandboxCore::GetMetadataFileName());
+	return ImportDetail::ImportFromZip(*InInspection.Reader, InInspection.AllFileNames, InSandboxTargetRoot, 
+		[&InInspection, &ManifestFilePath, &MetadataFilePath](const FString& InFileName) -> const TArray<uint8>* 
+		{
+			if (InFileName == ManifestFilePath)
+			{
+				return &InInspection.ManifestBytes;
+			}
+			if (InFileName == MetadataFilePath)
+			{
+				return &InInspection.MetadataBytes;
+			}
+			return nullptr;
+		});
+}
+}
diff --git a/Engine/Plugins/Developer/Sandbox/FileSandbox/Source/FileSandboxCore/Public/Utils/SandboxDirectoryUtils.h b/Engine/Plugins/Developer/Sandbox/FileSandbox/Source/FileSandboxCore/Public/Utils/SandboxDirectoryUtils.h
--- a/Engine/Plugins/Developer/Sandbox/FileSandbox/Source/FileSandboxCore/Public/Utils/SandboxDirectoryUtils.h
+++ b/Engine/Plugins/Developer/Sandbox/FileSandbox/Source/FileSandboxCore/Public/Utils/SandboxDirectoryUtils.h
@@ -18,11 +18,10 @@ namespace UE::FileSandboxCore
 enum class EBreakBehavior : uint8;
 using FProcessSandboxDirectorySignature = EBreakBehavior(const FString& InBaseDirectory);
 
-
-	
 /** @return The base directory into which new sandboxes are stored. */
 UE_API FString GetBaseSandboxDirectory();
 
+
 /**
  * Enumerates all directories that contain a sandbox. The search only considers the top-level directories.
  * @param InProcessSandbox Invokes for each found sandbox.
@@ -153,6 +152,9 @@ UE_API FText FormatDirectorySuitabilityAsReason(
  * @return Path to sandboxed file corresponding to the non-sandbox file, if the file has a sandboxed edit or add operation. Unset, otherwise.
  */
 UE_API TOptional<FString> GetSandboxPathFor(const FString& InSandboxDirectory, const FString& InNonSandboxPath);
+
+/** @return Name of the sandbox if the directory is the root of a sandbox. Unset otherwise. */
+UE_API TOptional<FString> GetSandboxName(const FString& InSandboxDirectory);
 }
 
 #undef UE_API
\ No newline at end of file
diff --git a/Engine/Plugins/Developer/Sandbox/FileSandbox/Source/FileSandboxCore/Public/Utils/ZipTransferUtils.h b/Engine/Plugins/Developer/Sandbox/FileSandbox/Source/FileSandboxCore/Public/Utils/ZipTransferUtils.h
new file mode 100644
--- /dev/null
+++ b/Engine/Plugins/Developer/Sandbox/FileSandbox/Source/FileSandboxCore/Public/Utils/ZipTransferUtils.h
@@ -0,0 +1,64 @@
+﻿// Copyright Epic Games, Inc. All Rights Reserved.
+
+#pragma once
+
+#include "Containers/Array.h"
+#include "Containers/UnrealString.h"
+#include "Data/SandboxMetaData.h"
+#include "Misc/Optional.h"
+#include "Templates/UniquePtr.h"
+
+class FString;
+class FZipArchiveReader;
+
+#define UE_API FILESANDBOXCORE_API
+
+namespace UE::FileSandboxCore::ZipUtils
+{
+/** Zips up a sandbox and saves it to the destination. */
+UE_API bool ExportSandboxToZip(const FString& InSandboxRoot, const FString& InOutputZipPath);
+
+/**
+ * Holds the result for inspecting whether a zip archive is suitable for importing.
+ * This also contains data produced as side effects, which can be reused for importing. 
+ */
+struct FImportInspectionResult
+{
+	/** The zip reader that was produced as side effect. You can re-use this for the import operation. */
+	TUniquePtr<FZipArchiveReader> Reader; // FZipArchiveReader does not support move semantics so we must wrap it with TUniquePtr.
+	/** All file names in the archive. Produced as side effect. */
+	TArray<FString> AllFileNames;
+	
+	/** The name of the sandbox contained in the zip. Produced as side effect. */
+	FString SandboxName;
+	/** Manifest file content. Produced as side effect. */
+	TArray<uint8> ManifestBytes;
+	/** Metadata file content. Produced as side effect. */
+	TArray<uint8> MetadataBytes;
+	/** Metadata. Produced as side effect. */
+	TOptional<FFileSandboxCore_SandboxMetaData> MetaData;
+	
+	explicit FImportInspectionResult(TUniquePtr<FZipArchiveReader> InReader);
+	
+	/** @return Whether the zip archive is valid to import from. */
+	bool IsValid() const { return !SandboxName.IsEmpty() && MetaData.IsSet(); }
+};
+
+/** 
+ * Inspects whether a zip contains a sandbox valid for importing, i.e.
+ * - the zip must be valid,
+ * - must contain exactly one root directory with the name of the sandbox,
+ * - must contain required files (manifest and metadata) 
+ * @return Result of inspecting the zip. Unset if I/O errors occurs.
+ * @note Creates IFileHandle that FZipArchiveReader takes ownership of until destroyed.
+ */
+UE_API TOptional<FImportInspectionResult> InspectFileForImport(const FString& InPathToZip);
+
+/** 
+ * Unzips a zip and imports it to the sandbox system reusing any data produced during inspection. 
+ * @note Pre-existing files in the target directory will be deleted.
+ */
+UE_API bool ImportSandboxFromZip(const FImportInspectionResult& InInspection, const FString& InSandboxTargetRoot);
+}
+
+#undef UE_API
