Sandboxed Editing: Fix "Show in file manager" context menu action not working for added files while not in sandbox.

#Synced-CL 50904440

diff --git a/Engine/Plugins/Developer/Sandbox/FileSandbox/Source/FileSandboxCore/Private/Utils/SandboxDirectoryUtils.cpp b/Engine/Plugins/Developer/Sandbox/FileSandbox/Source/FileSandboxCore/Private/Utils/SandboxDirectoryUtils.cpp
--- a/Engine/Plugins/Developer/Sandbox/FileSandbox/Source/FileSandboxCore/Private/Utils/SandboxDirectoryUtils.cpp
+++ b/Engine/Plugins/Developer/Sandbox/FileSandbox/Source/FileSandboxCore/Private/Utils/SandboxDirectoryUtils.cpp
@@ -299,9 +299,7 @@ TOptional<FString> GetSandboxPathFor(const FString& InSandboxDirectory, const FS
 			*SandboxContentDirectory, InAssetPath, InFilesystemPath
 			);
 		const TOptional<FSandboxedPlatformFilePath> Path = MakeSandboxedFilePathIfInMountPoint(InNonSandboxPath, MountPointPath);
-			
-		const bool bExists = Path && FPaths::FileExists(Path->GetNonSandboxPath());
-		if (bExists)
+		if (Path && !Path->GetSandboxPath().IsEmpty())
 		{
 			Result = Path->GetSandboxPath();
 			return EBreakBehavior::Break;
