File Sandbox: Add API for reverting file changes.

FSandboxPlatformFile::DiscardChanges has been split up into DiscardAll and RevertSandbox.
The difference between discarding and reverting is whether the files in the sandboxed internal file system are removed.
- Discarding retains them, reverting removes them. Discarding happens when the user leaves the sandbox.
- Reverting happens when the user wants to undo touching a file. Reverting implicitly means discarding, and on top also making sure the change does not show up in the future

#Synced-CL 51334859

diff --git a/Engine/Plugins/Developer/Concert/ConcertSync/ConcertSyncClient/Source/ConcertSyncClient/Private/Sandbox/ConcertPackageSandbox.cpp b/Engine/Plugins/Developer/Concert/ConcertSync/ConcertSyncClient/Source/ConcertSyncClient/Private/Sandbox/ConcertPackageSandbox.cpp
--- a/Engine/Plugins/Developer/Concert/ConcertSync/ConcertSyncClient/Source/ConcertSyncClient/Private/Sandbox/ConcertPackageSandbox.cpp
+++ b/Engine/Plugins/Developer/Concert/ConcertSync/ConcertSyncClient/Source/ConcertSyncClient/Private/Sandbox/ConcertPackageSandbox.cpp
@@ -15,7 +15,7 @@
 #include "Types/Manager/LoadSandboxByNameArgs.h"
 #include "Types/Manager/NewSandboxArgs.h"
 #include "Types/Manager/SandboxCreationResult.h"
-#include "Types/Sandbox/DiscardResult.h"
+#include "Types/Sandbox/RevertResult.h"
 
 namespace UE::ConcertSyncClient
 {
@@ -130,7 +130,7 @@ void FConcertPackageSandbox::DiscardSandbox(TArray<FName>& PackagesPendingHotRel
 	ISandboxManager& SandboxManager = IFileSandboxCoreModule::Get().GetSandboxManager();
 	if (ISandboxInstance* Sandbox = SandboxManager.GetActiveSandboxInstance())
 	{
-		const FDiscardResult DiscardResult = Sandbox->DiscardSandbox();
+		const FRevertResult DiscardResult = Sandbox->RevertSandbox();
 		PackagesPendingHotReload.Append(DiscardResult.PackagesPendingHotReload);
 		PackagesPendingPurge.Append(DiscardResult.PackagesPendingPurge);
 	}
diff --git a/Engine/Plugins/Developer/Sandbox/FileSandbox/Source/FileSandboxCore/Private/Sandbox/ISandboxInstance.cpp b/Engine/Plugins/Developer/Sandbox/FileSandbox/Source/FileSandboxCore/Private/Sandbox/ISandboxInstance.cpp
--- a/Engine/Plugins/Developer/Sandbox/FileSandbox/Source/FileSandboxCore/Private/Sandbox/ISandboxInstance.cpp
+++ b/Engine/Plugins/Developer/Sandbox/FileSandbox/Source/FileSandboxCore/Private/Sandbox/ISandboxInstance.cpp
@@ -4,8 +4,8 @@
 
 #include "Misc/DateTime.h"
 #include "Types/EBreakBehavior.h"
-#include "Types/Sandbox/DiscardArgs.h"
-#include "Types/Sandbox/DiscardResult.h"
+#include "Types/Sandbox/RevertArgs.h"
+#include "Types/Sandbox/RevertResult.h"
 #include "Types/Sandbox/PersistArgs.h"
 #include "Utils/PackageSandboxUtils.h"
 
@@ -19,9 +19,16 @@ bool ISandboxInstance::PersistAll()
 	return Result.PersistStatus == EPersistStatus::Success;
 }
 
-void ISandboxInstance::DiscardAll()
+void ISandboxInstance::RevertAll()
 {
-	const FDiscardResult Result = DiscardSandbox(FDiscardArgs());
+	const FRevertResult Result = RevertSandbox(FRevertArgs());
+	HotReloadPackages(Result.PackagesPendingHotReload);
+	PurgePackages(Result.PackagesPendingPurge);
+}
+
+void ISandboxInstance::RevertSpecified(const TConstArrayView<FString>& InFiles)
+{
+	const FRevertResult Result = RevertSandbox(FRevertArgs(InFiles));
 	HotReloadPackages(Result.PackagesPendingHotReload);
 	PurgePackages(Result.PackagesPendingPurge);
 }
diff --git a/Engine/Plugins/Developer/Sandbox/FileSandbox/Source/FileSandboxCore/Private/Sandbox/Persist/DiscardUtils.cpp b/Engine/Plugins/Developer/Sandbox/FileSandbox/Source/FileSandboxCore/Private/Sandbox/Persist/DiscardUtils.cpp
--- a/Engine/Plugins/Developer/Sandbox/FileSandbox/Source/FileSandboxCore/Private/Sandbox/Persist/DiscardUtils.cpp
+++ b/Engine/Plugins/Developer/Sandbox/FileSandbox/Source/FileSandboxCore/Private/Sandbox/Persist/DiscardUtils.cpp
@@ -15,54 +15,56 @@
 
 namespace UE::FileSandboxCore
 {
-void DiscardRemovedPaths(
-	const TSet<FSandboxPath>& InDeletedSandboxPaths, IPlatformFile& InLowerLevel, 
-	TArray<FName>& OutPackagesPendingPurge,
-	TFunctionRef<void(const FFileChange&)> InProcessFileChangeCallback
-	)
+namespace DiscardDetail
 {
-	AnalyzeFilesMarkedRemoved(InLowerLevel, InDeletedSandboxPaths,	
-		[&OutPackagesPendingPurge, &InProcessFileChangeCallback](const FNonSandboxPath& NonSandboxPath)
-		{
-			// If this file maps to a package then we need to flush its linker so that we can remove the file from the sandbox
-			FName PackageName;
-			FlushPackageFile(NonSandboxPath, &PackageName);
-				
-			if (!PackageName.IsNone())
-			{
-				OutPackagesPendingPurge.Add(PackageName);
-			}
-		
-			InProcessFileChangeCallback(FFileChange{ NonSandboxPath, EFileChangeAction::Added });
-			return EBreakBehavior::Continue;
-		});
+static bool ShouldDiscardFile(TConstArrayView<FString> InFilesToDiscard, const FString& InFile)
+{
+	return InFilesToDiscard.IsEmpty() || InFilesToDiscard.ContainsByPredicate([&InFile](const FString& InDiscardableFile)
+	{
+		return FPaths::IsSamePath(InDiscardableFile, InFile);
+	});
+}
+static bool ShouldDiscardFile(TConstArrayView<FString> InFilesToDiscard, const FSandboxedPlatformFilePath& InFile)
+{
+	return ShouldDiscardFile(InFilesToDiscard, InFile.GetNonSandboxPath());
+}
 }
 
 void DiscardAddedAndModifiedPaths(
+	TConstArrayView<FString> InFilesToDiscard, 
 	TConstArrayView<FSandboxMountPoint> InMountPoints, IPlatformFile& InLowerLevel, 
-	TArray<FName>& OutPackagesPendingHotReload,
-	TArray<FName>& OutPackagesPendingPurge,
-	TFunctionRef<void(const FFileChange&)> InProcessFileChangeCallback
+	TArray<FName>& OutPackagesPendingHotReload, TArray<FName>& OutPackagesPendingPurge,
+	TFunctionRef<void(const FSandboxedPlatformFilePath& Path, EFileChangeAction Action)> InProcessFileChangeCallback
 	)
 {
 	const auto HandleEditedFile = [&](const FSandboxedPlatformFilePath& InPath)
 	{
+		if (!DiscardDetail::ShouldDiscardFile(InFilesToDiscard, InPath))
+		{
+			return EBreakBehavior::Continue;
+		}
+			
 		// If this file maps to a package then we need to flush its linker so that we can remove the file from the sandbox
 		FName PackageName;
 		FlushPackageFile(InPath.GetNonSandboxPath(), &PackageName, false);
 			
 		if (!PackageName.IsNone())
 		{
-			OutPackagesPendingHotReload.Add(PackageName);
 			OutPackagesPendingPurge.Remove(PackageName);
+			OutPackagesPendingHotReload.Add(PackageName);
 		}
 				
-		InProcessFileChangeCallback(FFileChange{ InPath.GetNonSandboxPath(), EFileChangeAction::Modified });
+		InProcessFileChangeCallback(InPath, EFileChangeAction::Modified);
 		return EBreakBehavior::Continue;
 	};
 	
 	const auto HandleAddedFile = [&](const FSandboxedPlatformFilePath& InPath)
 	{
+		if (!DiscardDetail::ShouldDiscardFile(InFilesToDiscard, InPath))
+		{
+			return EBreakBehavior::Continue;
+		}
+			
 		// If this file maps to a package then we need to flush its linker so that we can remove the file from the sandbox
 		FName PackageName;
 		FlushPackageFile(InPath.GetNonSandboxPath(), &PackageName, false);
@@ -73,7 +75,7 @@ void DiscardAddedAndModifiedPaths(
 			OutPackagesPendingHotReload.Remove(PackageName);
 		}
 				
-		InProcessFileChangeCallback(FFileChange{ InPath.GetNonSandboxPath(), EFileChangeAction::Removed });
+		InProcessFileChangeCallback(InPath, EFileChangeAction::Removed);
 		return EBreakBehavior::Continue;
 	};
 	
diff --git a/Engine/Plugins/Developer/Sandbox/FileSandbox/Source/FileSandboxCore/Private/Sandbox/Persist/DiscardUtils.h b/Engine/Plugins/Developer/Sandbox/FileSandbox/Source/FileSandboxCore/Private/Sandbox/Persist/DiscardUtils.h
--- a/Engine/Plugins/Developer/Sandbox/FileSandbox/Source/FileSandboxCore/Private/Sandbox/Persist/DiscardUtils.h
+++ b/Engine/Plugins/Developer/Sandbox/FileSandbox/Source/FileSandboxCore/Private/Sandbox/Persist/DiscardUtils.h
@@ -10,19 +10,17 @@ class IPlatformFile;
 
 namespace UE::FileSandboxCore
 {
-struct FFileChange;
+enum class EFileChangeAction : uint8;
 struct FSandboxMountPoint;
-	
-void DiscardRemovedPaths(
-	const TSet<FSandboxPath>& InDeletedSandboxPaths, IPlatformFile& InLowerLevel, 
-	TArray<FName>& OutPackagesPendingPurge,
-	TFunctionRef<void(const FFileChange&)> InProcessFileChangeCallback = [](const FFileChange&){}
-	);
-	
+
+/**
+ * Discards files that have been added or edited in the sandbox.
+ * @param InFilesToDiscard Relative or absolute non-sandbox paths. If left empty, revert all files.  
+ */
 void DiscardAddedAndModifiedPaths(
+	TConstArrayView<FString> InFilesToDiscard, 
 	TConstArrayView<FSandboxMountPoint> InMountPoints, IPlatformFile& InLowerLevel, 
-	TArray<FName>& OutPackagesPendingHotReload,
-	TArray<FName>& OutPackagesPendingPurge,
-	TFunctionRef<void(const FFileChange&)> InProcessFileChangeCallback = [](const FFileChange&){}
+	TArray<FName>& OutPackagesPendingHotReload, TArray<FName>& OutPackagesPendingPurge,
+	TFunctionRef<void(const FSandboxedPlatformFilePath& Path, EFileChangeAction Action)> InProcessFileChangeCallback = [](const FSandboxedPlatformFilePath&, EFileChangeAction){}
 	);
 }
diff --git a/Engine/Plugins/Developer/Sandbox/FileSandbox/Source/FileSandboxCore/Private/Sandbox/Platform/SandboxPlatformFile.cpp b/Engine/Plugins/Developer/Sandbox/FileSandbox/Source/FileSandboxCore/Private/Sandbox/Platform/SandboxPlatformFile.cpp
--- a/Engine/Plugins/Developer/Sandbox/FileSandbox/Source/FileSandboxCore/Private/Sandbox/Platform/SandboxPlatformFile.cpp
+++ b/Engine/Plugins/Developer/Sandbox/FileSandbox/Source/FileSandboxCore/Private/Sandbox/Platform/SandboxPlatformFile.cpp
@@ -750,64 +750,153 @@ bool FSandboxPlatformFile::PersistSandbox(
 	return bFullOperationSuccess;
 }
 
-void FSandboxPlatformFile::DiscardSandbox(
-	TArray<FName>& OutPackagesPendingHotReload, 
-	TArray<FName>& OutPackagesPendingPurge,
+void FSandboxPlatformFile::DiscardAll(
+	TArray<FName>& OutPackagesPendingHotReload, TArray<FName>& OutPackagesPendingPurge,
 	TFunctionRef<void(const FFileChange&)> InProcessFileChangeCallback
 	)
 {
-	// We need to disable the sandbox while we do this
-	TGuardValue<TAtomic<bool>, bool> DisableSandboxGuard(bSandboxEnabled, false);
+	DiscardSandboxInternal(
+		OutPackagesPendingHotReload, OutPackagesPendingPurge, 
+		InProcessFileChangeCallback, 
+		[&InProcessFileChangeCallback](const FSandboxedPlatformFilePath& Path, EFileChangeAction Action)
+		{
+			InProcessFileChangeCallback(FFileChange(Path.GetNonSandboxPath(), Action));
+		});
+}
+
+void FSandboxPlatformFile::RevertSandbox(
+	TConstArrayView<FString> InFilesToRevert, 
+	TArray<FName>& OutPackagesPendingHotReload, TArray<FName>& OutPackagesPendingPurge, 
+	TFunctionRef<void(const FFileChange&)> InProcessFileChangeCallback
+	)
+{
+	// Additionally, also delete from the sandbox internal file system, so the change is not rediscovered when reloading the sandbox.
+	const auto HandleRevertAddedOrEditedFile =[this, &InProcessFileChangeCallback](const FSandboxedPlatformFilePath& Path, EFileChangeAction Action)
+	{
+		const FString& SandboxPath = Path.GetSandboxPath();
+		if (ensure(LowerLevel->FileExists(*SandboxPath)))
+		{
+			LowerLevel->DeleteFile(*SandboxPath);
+		}
+			
+		InProcessFileChangeCallback(FFileChange(Path.GetNonSandboxPath(), Action));
+	};
 	
+	DiscardSandboxInternal(
+		OutPackagesPendingHotReload, OutPackagesPendingPurge, 
+		InProcessFileChangeCallback, HandleRevertAddedOrEditedFile,
+		InFilesToRevert
+		);
+}
+
+bool FSandboxPlatformFile::DeletedPackageExistsInNonSandbox(FString InFilename) const
+{
+	const FSandboxedPlatformFilePath ResolvedPath = ToSandboxPath(MoveTemp(InFilename));
+	return ResolvedPath.HasSandboxPath()
+		&& IsPathDeleted(ResolvedPath)
+		&& LowerLevel->FileExists(*ResolvedPath.GetNonSandboxPath());
+}
+
+void FSandboxPlatformFile::DiscardSandboxInternal(
+	TArray<FName>& OutPackagesPendingHotReload, TArray<FName>& OutPackagesPendingPurge, 
+	TFunctionRef<void(const FFileChange&)> InProcessRemoved,
+	TFunctionRef<void(const FSandboxedPlatformFilePath& Path, EFileChangeAction Action)> InProcessAddedOrEdited,
+	TConstArrayView<FString> InFilesToDiscard
+	)
+{
+	// We need to disable the sandbox while we do this
+	const TGuardValue<TAtomic<bool>, bool> DisableSandboxGuard(bSandboxEnabled, false);
+
 #if UE_SANDBOX_WITH_DIRECTORY_WATCHER
 	TArray<FFileChangeData> FileChanges;
-#endif
-	const auto HandleFileChange = [&](const FFileChange& Change)
+	const auto EmplaceFileChange = [&FileChanges](const FFileChange& Change)
 	{
-#if UE_SANDBOX_WITH_DIRECTORY_WATCHER
 		FileChanges.Emplace(Change.ToDirectoryWatcherChange());
-#endif
-		InProcessFileChangeCallback(Change);
 	};
+#else 
+	const auto EmplaceFileChange = [](const FFileChange& Change) {};
+#endif
 
 	// Add any files that were deleted by the sandbox but exist in the non-sandbox directory
+	DiscardRemoved(InFilesToDiscard, OutPackagesPendingPurge, [&](const FString& NonSandboxPath)
 	{
-		FScopeLock DeletedSandboxPathsLock(&DeletedSandboxPathsCS);
-		DiscardRemovedPaths(DeletedSandboxPaths, *LowerLevel, OutPackagesPendingPurge, HandleFileChange);
-		DeletedSandboxPaths.Reset();
-	}
+		const FFileChange Change(NonSandboxPath, EFileChangeAction::Added);
+		EmplaceFileChange(Change);
+		InProcessRemoved(Change);
+	});
 
 	{
-		FScopeLock SandboxMountPointsLock(&SandboxMountPointsCS);
-		DiscardAddedAndModifiedPaths(SandboxMountPoints, *LowerLevel, OutPackagesPendingHotReload, OutPackagesPendingPurge, HandleFileChange);
+		const FScopeLock SandboxMountPointsLock(&SandboxMountPointsCS);
+		DiscardAddedAndModifiedPaths(
+			InFilesToDiscard, SandboxMountPoints, *LowerLevel, OutPackagesPendingHotReload, OutPackagesPendingPurge, 
+			[&](const FSandboxedPlatformFilePath& Path, EFileChangeAction Action)
+			{
+				EmplaceFileChange(FFileChange(Path.GetNonSandboxPath(), Action));
+				InProcessAddedOrEdited(Path, Action);
+			});
 	}
 
 #if UE_SANDBOX_WITH_DIRECTORY_WATCHER
 	// Notify that the sandboxed directories have been restored to their original state
-	if (FDirectoryWatcherModule* DirectoryWatcherModule = GetDirectoryWatcherModuleIfLoaded())
+	FDirectoryWatcherModule* DirectoryWatcherModule = GetDirectoryWatcherModuleIfLoaded();
+	if (DirectoryWatcherModule && FileChanges.Num() > 0)
 	{
-		if (FileChanges.Num() > 0)
-		{
-			DirectoryWatcherModule->RegisterExternalChanges(FileChanges);
-
-			// Force a sync of the Asset Registry immediately to process the
-			// file changes. This ensures that the Asset Registry is brought
-			// up to date especially after file/package deletes, since
-			// otherwise it may not get updated before packages are hot
-			// reloaded and stale entries could still be present in the
-			// registry.
-			SynchronizeAssetRegistry();
-		}
+		DirectoryWatcherModule->RegisterExternalChanges(FileChanges);
+
+		// Force a sync of the Asset Registry immediately to process the
+		// file changes. This ensures that the Asset Registry is brought
+		// up to date especially after file/package deletes, since
+		// otherwise it may not get updated before packages are hot
+		// reloaded and stale entries could still be present in the
+		// registry.
+		SynchronizeAssetRegistry();
 	}
 #endif
 }
 
-bool FSandboxPlatformFile::DeletedPackageExistsInNonSandbox(FString InFilename) const
+void FSandboxPlatformFile::DiscardRemoved(
+	TConstArrayView<FString> InFilesToDiscard, 
+	TArray<FName>& OutPackagesPendingPurge, 
+	TFunctionRef<void(const FString&)> InProcessRemoved
+	)
 {
-	const FSandboxedPlatformFilePath ResolvedPath = ToSandboxPath(MoveTemp(InFilename));
-	return ResolvedPath.HasSandboxPath()
-		&& IsPathDeleted(ResolvedPath)
-		&& LowerLevel->FileExists(*ResolvedPath.GetNonSandboxPath());
+	const FScopeLock DeletedSandboxPathsLock(&DeletedSandboxPathsCS);
+	const auto DiscardSingle = [&OutPackagesPendingPurge, &InProcessRemoved](const FString& NonSandboxPath)
+	{
+		// If this file maps to a package then we need to flush its linker so that we can remove the file from the sandbox
+		FName PackageName;
+		FlushPackageFile(NonSandboxPath, &PackageName);
+    
+		if (!PackageName.IsNone())
+		{
+			OutPackagesPendingPurge.Add(PackageName);
+		}
+    
+		InProcessRemoved(NonSandboxPath);
+	};
+	
+	const bool bClearAll = InFilesToDiscard.IsEmpty();
+	if (bClearAll)
+	{
+		for (const FString& NonSandboxPath : DeletedSandboxPaths)
+		{
+			DiscardSingle(NonSandboxPath);
+		}
+		DeletedSandboxPaths.Empty();
+	}
+	else
+	{
+		for (const FString& NonSandboxPath : InFilesToDiscard)
+		{
+			const FString AbsPath = FPaths::ConvertRelativePathToFull(NonSandboxPath);
+			const bool bIsMarkedRemoved = DeletedSandboxPaths.Contains(AbsPath);
+			if (bIsMarkedRemoved)
+			{
+				DiscardSingle(AbsPath);
+				DeletedSandboxPaths.Remove(AbsPath);
+			}
+		}
+	}
 }
 
 void FSandboxPlatformFile::OnContentPathMounted(const FString& InAssetPath, const FString& InFilesystemPath)
diff --git a/Engine/Plugins/Developer/Sandbox/FileSandbox/Source/FileSandboxCore/Private/Sandbox/Platform/SandboxPlatformFile.h b/Engine/Plugins/Developer/Sandbox/FileSandbox/Source/FileSandboxCore/Private/Sandbox/Platform/SandboxPlatformFile.h
--- a/Engine/Plugins/Developer/Sandbox/FileSandbox/Source/FileSandboxCore/Private/Sandbox/Platform/SandboxPlatformFile.h
+++ b/Engine/Plugins/Developer/Sandbox/FileSandbox/Source/FileSandboxCore/Private/Sandbox/Platform/SandboxPlatformFile.h
@@ -20,7 +20,6 @@ namespace UE::FileSandboxCore
 class IPersistFeedback;
 enum class EFileChangeAction : uint8;
 enum class ESandboxFileChange : uint8;
-struct FDiscardSandboxArgs;
 struct FFileChange;
 struct FPersistArgs;
 struct FPersistResult;
@@ -91,13 +90,44 @@ public:
 		);
 
 	/** 
-	 *	Discard the sandbox state
-	 *	This will trigger directory watcher notifications for files that are restored.
-	 *	This will also gather packages that need to be hot reloaded or purge from memory
+	 * Brings all files back to the state they had before the sandbox:
+	 * - notifies the asset registry of the file change
+	 * - unloads any FLinker instances that may depend on it
+	 * - you'll have to manually hot-reload and purge packages.
+	 * This function is useful if the user is leaving the sandbox so when the sandbox is reloaded removed, added & edited files are retained.
+	 * 
+	 * @see RevertSandbox: The difference to RevertSandbox is that added and edited files remain in the sandbox-internal file system, 
+	 * e.g. Intermediate/Sandboxes. This means the next time this sandbox is reloaded, the changes can be discovered.
+	 * 
+	 * @param OutPackagesPendingHotReload Packages that need to be hot reloaded.
+	 * @param OutPackagesPendingPurge Packages that need to be purged.
+	 * @param InProcessFileChangeCallback Invoked on non-sandbox paths, and with the file action that reverts the sandbox action. 
+	 *	E.g. a file removed in sandbox would be EFileChangeAction::Added.
 	 */
-	void DiscardSandbox(
-		TArray<FName>& OutPackagesPendingHotReload,
-		TArray<FName>& OutPackagesPendingPurge,
+	void DiscardAll(
+		TArray<FName>& OutPackagesPendingHotReload, TArray<FName>& OutPackagesPendingPurge,
+		TFunctionRef<void(const FFileChange&)> InProcessFileChangeCallback = [](const FFileChange&){}
+		);
+	
+	/** 
+	 * Brings all files back to the state they had before the sandbox:
+	 * - notifies the asset registry of the file change
+	 * - unloads any FLinker instances that may depend on it
+	 * - you'll have to manually hot-reload and purge packages.
+	 * This function is useful if the user remains in the sandbox and wants to undo a file action.
+	 * 
+	 * @see DiscardAll: The difference to DiscardAll is that added and edited files are removed from the sandbox-internal file system, 
+	 * e.g. Intermediate/Sandboxes. This means the next time this sandbox is reloaded, the changes will NOT be discovered.
+	 * 
+	 * @param InFilesToRevert Absolute or relative non-sandbox paths. If empty, discard all changes.
+	 * @param OutPackagesPendingHotReload Packages that need to be hot reloaded.
+	 * @param OutPackagesPendingPurge Packages that need to be purged.
+	 * @param InProcessFileChangeCallback Invoked on non-sandbox paths, and with the file action that reverts the sandbox action. 
+	 *	E.g. a file removed in sandbox would be EFileChangeAction::Added.
+	 */
+	void RevertSandbox(
+		TConstArrayView<FString> InFilesToRevert,
+		TArray<FName>& OutPackagesPendingHotReload, TArray<FName>& OutPackagesPendingPurge,
 		TFunctionRef<void(const FFileChange&)> InProcessFileChangeCallback = [](const FFileChange&){}
 		);
 	
@@ -114,6 +144,32 @@ private:
 		FString Path;
 		FFileStatData StatData;
 	};
+
+	/**
+	 * Brings all files back to the state they had before the sandbox:
+	 * - notifies the asset registry of the file change
+	 * - unloads any FLinker instances that may depend on it
+	 * - you'll have to manually hot-reload and purge packages.
+	 * 
+	 * @param OutPackagesPendingHotReload Packages that need to be hot reloaded.
+	 * @param OutPackagesPendingPurge Packages that need to be purged.
+	 * @param InProcessRemoved Invoked on non-sandbox paths, and with the file action that reverts the sandbox action. 
+	 * @param InProcessAddedOrEdited Invoked with the file action that reverts the sandbox action. 
+	 * @param InFilesToDiscard Absolute or relative non-sandbox paths to discard. If empty, discard all files.
+	 */
+	void DiscardSandboxInternal(
+		TArray<FName>& OutPackagesPendingHotReload, TArray<FName>& OutPackagesPendingPurge,
+		TFunctionRef<void(const FFileChange&)> InProcessRemoved = [](const FFileChange&){},
+		TFunctionRef<void(const FSandboxedPlatformFilePath& Path, EFileChangeAction Action)> InProcessAddedOrEdited = [](const FSandboxedPlatformFilePath&, EFileChangeAction){},
+		TConstArrayView<FString> InFilesToDiscard = {}
+		);
+	
+	/** Utility for discarding files marked as removed. */
+	void DiscardRemoved(
+		TConstArrayView<FString> InFilesToDiscard,
+		TArray<FName>& OutPackagesPendingPurge,
+		TFunctionRef<void(const FString&)> InProcessRemoved = [](const FString&){}
+		);
 	
 	/** Callback when a new content path is mounted */
 	void OnContentPathMounted(const FString& InAssetPath, const FString& InFilesystemPath);
diff --git a/Engine/Plugins/Developer/Sandbox/FileSandbox/Source/FileSandboxCore/Private/Sandbox/SandboxInstance.cpp b/Engine/Plugins/Developer/Sandbox/FileSandbox/Source/FileSandboxCore/Private/Sandbox/SandboxInstance.cpp
--- a/Engine/Plugins/Developer/Sandbox/FileSandbox/Source/FileSandboxCore/Private/Sandbox/SandboxInstance.cpp
+++ b/Engine/Plugins/Developer/Sandbox/FileSandbox/Source/FileSandboxCore/Private/Sandbox/SandboxInstance.cpp
@@ -9,7 +9,7 @@
 #include "Persist/LogPersistFeedback.h"
 #include "Platform/SandboxPlatformFile.h"
 #include "Types/Manager/SandboxInitArgs.h"
-#include "Types/Sandbox/DiscardResult.h"
+#include "Types/Sandbox/RevertResult.h"
 #include "Types/SandboxFileChange.h"
 #include "UObject/Package.h"
 #include "Utils/AssetUtils.h"
@@ -22,7 +22,11 @@ namespace UE::FileSandboxCore
 {
 namespace Private
 {
-static void TrackRevertedFileChangeInManifest(FFileSandboxCore_ManifestData& OutManifestData, const FFileChange& InChange)
+/**
+ * Applies file actions resulting from reverting files to the manifest file.
+ * E.g. EFileChangeAction::Added implies that a sandboxed remove operation was reverted (so add it to the manifest).
+ */
+static void AddDiscardedFileChangeInManifest(FFileSandboxCore_ManifestData& OutManifestData, const FFileChange& InChange)
 {
 	switch (InChange.Action)
 	{
@@ -33,6 +37,21 @@ static void TrackRevertedFileChangeInManifest(FFileSandboxCore_ManifestData& Out
 	}
 }
 
+/**
+ * Applies file actions resulting from reverting files to the manifest file. 
+ * E.g. EFileChangeAction::Added implies that a sandboxed remove operation was reverted (so removed it from the manifest).
+ */
+static void RemoveRevertedFileChangeFromManifest(FFileSandboxCore_ManifestData& OutManifestData, const FFileChange& InChange)
+{
+	switch (InChange.Action)
+	{
+	case EFileChangeAction::Added: OutManifestData.DeletedFiles.Remove(InChange.Filename); break;
+	case EFileChangeAction::Modified: OutManifestData.ModifiedFiles.Remove(InChange.Filename); break;
+	case EFileChangeAction::Removed: OutManifestData.AddedFiles.Remove(InChange.Filename); break;
+	default: break;
+	}
+}
+
 static void TrackFileChangeInManifest(FFileSandboxCore_ManifestData& OutManifestData, const FSandboxedFileChangeInfo& InChange)
 {
 	switch (InChange.Action)
@@ -49,7 +68,6 @@ static void TrackFileChangeInManifest(FFileSandboxCore_ManifestData& OutManifest
 	}
 }
 
-
 /** 
  * Adds packages with in-memory changes. 
  * 
@@ -109,9 +127,9 @@ FSandboxInstance::~FSandboxInstance()
 	
 	ManifestContent.Empty();
 	TArray<FName> PackagesPendingHotReload, PackagesPendingPurge;
-	PlatformFile->DiscardSandbox(PackagesPendingHotReload, PackagesPendingPurge, [this](const FFileChange& Change)
+	PlatformFile->DiscardAll(PackagesPendingHotReload, PackagesPendingPurge, [this](const FFileChange& Change)
 	{
-		Private::TrackRevertedFileChangeInManifest(ManifestContent, Change); 
+		Private::AddDiscardedFileChangeInManifest(ManifestContent, Change); 
 	});
 	PlatformFile.Reset();
 
@@ -174,12 +192,13 @@ FPersistResult FSandboxInstance::PersistSandbox(const FPersistArgs& InArgs)
 	};
 }
 
-FDiscardResult FSandboxInstance::DiscardSandbox(const FDiscardArgs& InArgs)
+FRevertResult FSandboxInstance::RevertSandbox(const FRevertArgs& InArgs)
 {
-	FDiscardResult Result;
-	PlatformFile->DiscardSandbox(Result.PackagesPendingHotReload, Result.PackagesPendingPurge);
-	
-	ManifestContent.Empty();
+	FRevertResult Result;
+	PlatformFile->RevertSandbox(
+		InArgs.FilesToRevert, Result.PackagesPendingHotReload, Result.PackagesPendingPurge, 
+		[this](const FFileChange& Change) { Private::RemoveRevertedFileChangeFromManifest(ManifestContent, Change);}
+		);
 	SaveManifest();
 	
 	// Discarding has removed all file states, so inform anyone that cares about it.
diff --git a/Engine/Plugins/Developer/Sandbox/FileSandbox/Source/FileSandboxCore/Private/Sandbox/SandboxInstance.h b/Engine/Plugins/Developer/Sandbox/FileSandbox/Source/FileSandboxCore/Private/Sandbox/SandboxInstance.h
--- a/Engine/Plugins/Developer/Sandbox/FileSandbox/Source/FileSandboxCore/Private/Sandbox/SandboxInstance.h
+++ b/Engine/Plugins/Developer/Sandbox/FileSandbox/Source/FileSandboxCore/Private/Sandbox/SandboxInstance.h
@@ -34,14 +34,14 @@ public:
 		);
 	virtual ~FSandboxInstance() override;
 
-	using ISandboxInstance::DiscardSandbox;
+	using ISandboxInstance::RevertSandbox;
 
 	const FFileSandboxCore_ManifestData& GetManifestContent() const { return ManifestContent; }
 	IPlatformFile& GetLowerLevelPlatformFile() const { return *PlatformFile->GetLowerLevel(); }
 	
 	//~ Begin ISandboxInstance Interface
 	virtual FPersistResult PersistSandbox(const FPersistArgs& InArgs) override;
-	virtual FDiscardResult DiscardSandbox(const FDiscardArgs& InArgs) override;
+	virtual FRevertResult RevertSandbox(const FRevertArgs& InArgs) override;
 	virtual void EnumerateFileChanges(TFunctionRef<FProcessFileChangeSignature> InProcess, EFileEnumerationFlags InFlags) const override;
 	virtual TOptional<FDateTime> GetSandboxedFileTimestamp(const FString& InFilePath) const override;
 	virtual bool DeletedPackageExistsInNonSandbox(const FString& InFilename) const override;
diff --git a/Engine/Plugins/Developer/Sandbox/FileSandbox/Source/FileSandboxCore/Public/ISandboxInstance.h b/Engine/Plugins/Developer/Sandbox/FileSandbox/Source/FileSandboxCore/Public/ISandboxInstance.h
--- a/Engine/Plugins/Developer/Sandbox/FileSandbox/Source/FileSandboxCore/Public/ISandboxInstance.h
+++ b/Engine/Plugins/Developer/Sandbox/FileSandbox/Source/FileSandboxCore/Public/ISandboxInstance.h
@@ -3,7 +3,7 @@
 #pragma once
 
 #include "Types/GatheredFileChanges.h"
-#include "Types/Sandbox/DiscardArgs.h"
+#include "Types/Sandbox/RevertArgs.h"
 #include "Types/Sandbox/PersistResult.h"
 #include "Types/ArgsFwds.h"
 #include "Types/SandboxedFileChangeInfo.h"
@@ -16,7 +16,7 @@ namespace UE::FileSandboxCore
 {
 enum class EFileChangeGatherFlags : uint8;
 enum class EBreakBehavior : uint8;
-struct FDiscardResult;
+struct FRevertResult;
 struct FPersistArgs;
 struct FPersistedFileInfo;
 
@@ -43,19 +43,23 @@ public:
 	 */
 	virtual FPersistResult PersistSandbox(const FPersistArgs& InArgs) = 0;
 
-
 	
-	/** Discard the sandbox state reverting all files to the state they had before the sandbox. */
-	UE_API void DiscardAll();
+	/** Reverts all files to the state they had before the sandbox. */
+	UE_API void RevertAll();
+	
+	/** Reverts the files you specify to the state they had before the sandbox. */
+	UE_EXPERIMENTAL(5.8, "Reverting of single files does not yet work for all cases. See UE-368478.")
+	UE_API void RevertSpecified(const TConstArrayView<FString>& InFiles);
 	
 	/**
-	 * More advanced version of DiscardAll that gives you control when packages are hot-reloaded and purged.
-	 * This will also gather packages that need to be hot reloaded or purge from memory.
-	 * 
-	 * @param InArgs Input for the discard operation
-	 * @return Information about the discard operation; you must handle hot-reloading and purging yourself.
+	 * More advanced version that gives you full control: specifically when packages are hot-reloaded and purged.
+	 * This gathers the packages that you must hot-reload or purge from memory.
+	 * @param InArgs Input for the revert operation
+	 * @return Information about the revert operation; you must handle hot-reloading and purging yourself.
+	 * @see HotReloadPackages, PurgePackages
 	 */
-	[[nodiscard]] virtual FDiscardResult DiscardSandbox(const FDiscardArgs& InArgs = {}) = 0;
+	UE_EXPERIMENTAL(5.8, "Reverting of single files does not yet work for all cases. See UE-368478.")
+	[[nodiscard]] virtual FRevertResult RevertSandbox(const FRevertArgs& InArgs = {}) = 0;
 	
 	
 	/** Enumerates all files that have been changed (added, removed, or edited). */
diff --git a/Engine/Plugins/Developer/Sandbox/FileSandbox/Source/FileSandboxCore/Public/Types/Sandbox/DiscardArgs.h b/Engine/Plugins/Developer/Sandbox/FileSandbox/Source/FileSandboxCore/Public/Types/Sandbox/DiscardArgs.h
deleted file mode 100644
diff --git a/Engine/Plugins/Developer/Sandbox/FileSandbox/Source/FileSandboxCore/Public/Types/Sandbox/DiscardResult.h b/Engine/Plugins/Developer/Sandbox/FileSandbox/Source/FileSandboxCore/Public/Types/Sandbox/DiscardResult.h
deleted file mode 100644
diff --git a/Engine/Plugins/Developer/Sandbox/FileSandbox/Source/FileSandboxCore/Public/Types/Sandbox/RevertArgs.h b/Engine/Plugins/Developer/Sandbox/FileSandbox/Source/FileSandboxCore/Public/Types/Sandbox/RevertArgs.h
new file mode 100644
--- /dev/null
+++ b/Engine/Plugins/Developer/Sandbox/FileSandbox/Source/FileSandboxCore/Public/Types/Sandbox/RevertArgs.h
@@ -0,0 +1,28 @@
+﻿// Copyright Epic Games, Inc. All Rights Reserved.
+
+#pragma once
+
+#include "Containers/ArrayView.h"
+
+class FString;
+
+namespace UE::FileSandboxCore
+{
+/** Input arguments for reverting a sandbox. */
+struct UE_EXPERIMENTAL(5.8, "Reverting of single files does not yet work for all cases. See UE-368478.") FRevertArgs
+{
+	/** 
+	 * Non-sandbox paths to files that should be reverted.
+	 * If this is left empty, then the operation reverts all file changes.
+	 * 
+	 * This can be paths to files with any type of change (i.e. removed, edited, added).
+	 * The path can be absolute or relative.
+	 */
+	TConstArrayView<FString> FilesToRevert;
+	
+	FRevertArgs() = default;
+	explicit FRevertArgs(TConstArrayView<FString> InFilesToRevert) 
+		: FilesToRevert(InFilesToRevert) 
+	{}
+};
+}
diff --git a/Engine/Plugins/Developer/Sandbox/FileSandbox/Source/FileSandboxCore/Public/Types/Sandbox/RevertResult.h b/Engine/Plugins/Developer/Sandbox/FileSandbox/Source/FileSandboxCore/Public/Types/Sandbox/RevertResult.h
new file mode 100644
--- /dev/null
+++ b/Engine/Plugins/Developer/Sandbox/FileSandbox/Source/FileSandboxCore/Public/Types/Sandbox/RevertResult.h
@@ -0,0 +1,25 @@
+﻿// Copyright Epic Games, Inc. All Rights Reserved.
+
+#pragma once
+
+#include "Containers/Array.h"
+#include "UObject/NameTypes.h"
+
+namespace UE::FileSandboxCore
+{
+/** Result of reverting a sandbox. */
+struct UE_EXPERIMENTAL(5.8, "Reverting of single files does not yet work for all cases. See UE-368478.") FRevertResult
+{
+	/**
+	 * The package names that should be hot reloaded.
+	 * You must call HotReloadPackages on them.
+	 */
+	TArray<FName> PackagesPendingHotReload;
+	
+	/**
+	 * The package names that should be purged.
+	 * You must call PurgePackages on them.
+	 */
+	TArray<FName> PackagesPendingPurge;
+};
+}
diff --git a/Engine/Plugins/Developer/Sandbox/SandboxedEditing/Source/SandboxedEditing/Private/Framework/Models/SandboxSystemModel.cpp b/Engine/Plugins/Developer/Sandbox/SandboxedEditing/Source/SandboxedEditing/Private/Framework/Models/SandboxSystemModel.cpp
--- a/Engine/Plugins/Developer/Sandbox/SandboxedEditing/Source/SandboxedEditing/Private/Framework/Models/SandboxSystemModel.cpp
+++ b/Engine/Plugins/Developer/Sandbox/SandboxedEditing/Source/SandboxedEditing/Private/Framework/Models/SandboxSystemModel.cpp
@@ -317,7 +317,7 @@ void FSandboxSystemModel::DiscardAllChanges() const
 	
 	if (CanDiscardAllChanges() && ensure(Instance))
 	{
-		Instance->DiscardAll();
+		Instance->RevertAll();
 	}
 }
 
