Sandboxed Editing: Fix experimental warnings

#Synced-CL 51334859

diff --git a/Engine/Plugins/Developer/Sandbox/FileSandbox/Source/FileSandboxCore/Private/Sandbox/ISandboxInstance.cpp b/Engine/Plugins/Developer/Sandbox/FileSandbox/Source/FileSandboxCore/Private/Sandbox/ISandboxInstance.cpp
--- a/Engine/Plugins/Developer/Sandbox/FileSandbox/Source/FileSandboxCore/Private/Sandbox/ISandboxInstance.cpp
+++ b/Engine/Plugins/Developer/Sandbox/FileSandbox/Source/FileSandboxCore/Private/Sandbox/ISandboxInstance.cpp
@@ -19,19 +19,23 @@ bool ISandboxInstance::PersistAll()
 	return Result.PersistStatus == EPersistStatus::Success;
 }
 
+PRAGMA_DISABLE_EXPERIMENTAL_WARNINGS
 void ISandboxInstance::RevertAll()
 {
 	const FRevertResult Result = RevertSandbox(FRevertArgs());
 	HotReloadPackages(Result.PackagesPendingHotReload);
 	PurgePackages(Result.PackagesPendingPurge);
 }
+PRAGMA_ENABLE_EXPERIMENTAL_WARNINGS
 
+PRAGMA_DISABLE_EXPERIMENTAL_WARNINGS
 void ISandboxInstance::RevertSpecified(const TConstArrayView<FString>& InFiles)
 {
 	const FRevertResult Result = RevertSandbox(FRevertArgs(InFiles));
 	HotReloadPackages(Result.PackagesPendingHotReload);
 	PurgePackages(Result.PackagesPendingPurge);
 }
+PRAGMA_ENABLE_EXPERIMENTAL_WARNINGS
 
 FGatheredFileChanges ISandboxInstance::GatherChangedFiles(EFileChangeGatherFlags InFlags) const
 {
diff --git a/Engine/Plugins/Developer/Sandbox/FileSandbox/Source/FileSandboxCore/Public/ISandboxInstance.h b/Engine/Plugins/Developer/Sandbox/FileSandbox/Source/FileSandboxCore/Public/ISandboxInstance.h
--- a/Engine/Plugins/Developer/Sandbox/FileSandbox/Source/FileSandboxCore/Public/ISandboxInstance.h
+++ b/Engine/Plugins/Developer/Sandbox/FileSandbox/Source/FileSandboxCore/Public/ISandboxInstance.h
@@ -59,7 +59,9 @@ public:
 	 * @see HotReloadPackages, PurgePackages
 	 */
 	UE_EXPERIMENTAL(5.8, "Reverting of single files does not yet work for all cases. See UE-368478.")
+	PRAGMA_DISABLE_EXPERIMENTAL_WARNINGS
 	[[nodiscard]] virtual FRevertResult RevertSandbox(const FRevertArgs& InArgs = {}) = 0;
+	PRAGMA_ENABLE_EXPERIMENTAL_WARNINGS
 	
 	
 	/** Enumerates all files that have been changed (added, removed, or edited). */
