Sandboxes: Fix importing a sandbox resetting the timestamp of added & edited files

#Synced-CL 52283655

#Synced-CL 52283655

diff --git a/Engine/Plugins/Developer/Sandbox/FileSandbox/Source/FileSandboxCore/Private/Utils/ZipTransferUtils.cpp b/Engine/Plugins/Developer/Sandbox/FileSandbox/Source/FileSandboxCore/Private/Utils/ZipTransferUtils.cpp
--- a/Engine/Plugins/Developer/Sandbox/FileSandbox/Source/FileSandboxCore/Private/Utils/ZipTransferUtils.cpp
+++ b/Engine/Plugins/Developer/Sandbox/FileSandbox/Source/FileSandboxCore/Private/Utils/ZipTransferUtils.cpp
@@ -14,6 +14,7 @@
 #include "Misc/Paths.h"
 #include "SandboxFileUtils.h"
 #include "FileChange/EFileChangeAction.h"
+#include "FileUtilities/ZipFileMetaData.h"
 #include "Utils/SandboxDirectoryUtils.h"
 
 #if UE_SANDBOX_WITH_DIRECTORY_WATCHER
@@ -50,7 +51,13 @@ static bool WriteToZip(const FString& InSandboxRoot, IFileHandle& InZipFile, TCo
 		// This converts it to "MySandbox/manifest.json".
 		FString RelativeToRoot = FileToCompress;
 		FPaths::MakePathRelativeTo(RelativeToRoot, *InSandboxRoot);
-		ZipArchiveWriter.AddFile(RelativeToRoot, RawData, FDateTime::Now());
+		
+		// The file in the archive needs to have the same timestamp as in the file system so the edit time is correct when importing
+		const FDateTime Timestamp = FileManager.GetTimeStamp(*FileToCompress);
+		const bool bHasTimestamp = Timestamp != FDateTime::MinValue(); 
+		const FDateTime TimestampToSet = bHasTimestamp ? Timestamp : FDateTime::Now();
+		
+		ZipArchiveWriter.AddFile(RelativeToRoot, RawData, TimestampToSet);
 	}
 	
 	return true;
@@ -225,7 +232,7 @@ TOptional<FImportInspectionResult> InspectFileForImport(const FString& InPathToZ
 
 namespace ImportDetail
 {
-static bool ImportZippedFile(const TArray<uint8>& InFileContent, const FString& InTargetFile)
+static bool ImportZippedFile(const TArray<uint8>& InFileContent, const FString& InTargetFile, const TOptional<FDateTime>& InTimestamp)
 {
 	IPlatformFile& PlatformFile = FPlatformFileManager::Get().GetPlatformFile();
 	const FString Directory = FPaths::GetPath(InTargetFile);
@@ -245,10 +252,10 @@ static bool ImportZippedFile(const TArray<uint8>& InFileContent, const FString&
 	Archive->Close(); // Always explicitly close to catch errors from flush/close
 	const bool bIsSuccess = !Archive->IsError() && !Archive->IsCriticalError();
 	
-	if (bIsSuccess)
+	// The timestamp of created file should be equal to timestamp the file has in the zip archive. 
+	if (bIsSuccess && InTimestamp)
 	{
-		// TODO UE-366795: Set the timestamp of created file to be equal to timestamp the file has in the zip archive. 
-		// Currently, FZipArchiveReader has no API to get that information.
+		FileManager.SetTimeStamp(*InTargetFile, *InTimestamp);
 	}
 	
 	return bIsSuccess;
@@ -321,15 +328,17 @@ static bool ImportFromZip(
 		// Optimization: Is it a file that we already deserialized during inspection?
 		if (const TArray<uint8>* CachedContent = InGetCachedFileContent(ZippedFile))
 		{
-			bFreeOfErrors &= ImportZippedFile(*CachedContent, Destination);
+			const TOptional<FDateTime> NoTimestamp; // We don't care for timestamps on the manifest or metadata file
+			bFreeOfErrors &= ImportZippedFile(*CachedContent, Destination, NoTimestamp);
 			FilesImportedSoFar.Add(Destination);
 			continue;
 		}
 		
 		FileContent.Empty(FileContent.Num());
-		if (InReader.TryReadFile(ZippedFile, FileContent))
+		FileUtilities::FZipFileMetaData Metadata;
+		if (InReader.TryReadFile(ZippedFile, FileContent, nullptr, &Metadata))
 		{
-			bFreeOfErrors &= ImportZippedFile(FileContent, Destination);
+			bFreeOfErrors &= ImportZippedFile(FileContent, Destination, Metadata.Timestamp);
 			FilesImportedSoFar.Add(Destination);
 		}
 		else
diff --git a/Engine/Source/Developer/FileUtilities/FileUtilities.Build.cs b/Engine/Source/Developer/FileUtilities/FileUtilities.Build.cs
--- a/Engine/Source/Developer/FileUtilities/FileUtilities.Build.cs
+++ b/Engine/Source/Developer/FileUtilities/FileUtilities.Build.cs
@@ -6,7 +6,7 @@ public class FileUtilities : ModuleRules
 {
 	public FileUtilities(ReadOnlyTargetRules Target) : base(Target)
 	{
-		PrivateDependencyModuleNames.AddRange(new string[] { "Core"});
+		PublicDependencyModuleNames.AddRange(new string[] { "Core" });
 
 		if (Target.bBuildEditor)
 		{
diff --git a/Engine/Source/Developer/FileUtilities/Private/ZipArchiveReader.cpp b/Engine/Source/Developer/FileUtilities/Private/ZipArchiveReader.cpp
--- a/Engine/Source/Developer/FileUtilities/Private/ZipArchiveReader.cpp
+++ b/Engine/Source/Developer/FileUtilities/Private/ZipArchiveReader.cpp
@@ -6,6 +6,7 @@
 
 #include "Containers/StringConv.h"
 #include "Containers/Map.h"
+#include "FileUtilities/ZipFileMetaData.h"
 #include "GenericPlatform/GenericPlatformFile.h"
 #include "libzip/zip.h"
 #include "Logging/LogMacros.h"
@@ -35,7 +36,7 @@ public:
 
 	bool IsValid() const;
 	TArray<FString> GetFileNames() const;
-	bool TryReadFile(FStringView FileName, TArray<uint8>& OutData, FOutputDevice* ErrorHandler) const;
+	bool TryReadFile(FStringView FileName, TArray<uint8>& OutData, FOutputDevice* ErrorHandler, UE::FileUtilities::FZipFileMetaData* OutMetadata) const;
 };
 
 
@@ -57,9 +58,11 @@ TArray<FString> FZipArchiveReader::GetFileNames() const
 	return Impl->GetFileNames();
 }
 
-bool FZipArchiveReader::TryReadFile(FStringView FileName, TArray<uint8>& OutData, FOutputDevice* ErrorHandler) const
+bool FZipArchiveReader::TryReadFile(
+	FStringView FileName, TArray<uint8>& OutData, FOutputDevice* ErrorHandler, UE::FileUtilities::FZipFileMetaData* OutMetadata
+	) const
 {
-	return Impl->TryReadFile(FileName, OutData, ErrorHandler);
+	return Impl->TryReadFile(FileName, OutData, ErrorHandler, OutMetadata);
 }
 
 FZipArchiveReader::FImpl::FImpl(IFileHandle* InFileHandle, FOutputDevice* ErrorHandler)
@@ -182,7 +185,9 @@ TArray<FString> FZipArchiveReader::FImpl::GetFileNames() const
 	return Result;
 }
 
-bool FZipArchiveReader::FImpl::TryReadFile(FStringView FileName, TArray<uint8>& OutData, FOutputDevice* ErrorHandler) const
+bool FZipArchiveReader::FImpl::TryReadFile(
+	FStringView FileName, TArray<uint8>& OutData, FOutputDevice* ErrorHandler, UE::FileUtilities::FZipFileMetaData* OutMetadata
+	) const
 {
 	OutData.Reset();
 
@@ -259,6 +264,11 @@ bool FZipArchiveReader::FImpl::TryReadFile(FStringView FileName, TArray<uint8>&
 		OutData.Reset();
 		return false;
 	}
+	
+	if (OutMetadata && ZipFileStat.valid & ZIP_STAT_MTIME)
+	{
+		OutMetadata->Timestamp = FDateTime::FromUnixTimestamp(ZipFileStat.mtime);
+	}
 
 	return true;
 }
diff --git a/Engine/Source/Developer/FileUtilities/Public/FileUtilities/ZipArchiveReader.h b/Engine/Source/Developer/FileUtilities/Public/FileUtilities/ZipArchiveReader.h
--- a/Engine/Source/Developer/FileUtilities/Public/FileUtilities/ZipArchiveReader.h
+++ b/Engine/Source/Developer/FileUtilities/Public/FileUtilities/ZipArchiveReader.h
@@ -9,6 +9,11 @@
 #include "Containers/UnrealString.h"
 #include "Templates/UniquePtr.h"
 
+namespace UE::FileUtilities
+{
+struct FZipFileMetaData;
+}
+
 class FOutputDevice;
 class IFileHandle;
 
@@ -29,7 +34,10 @@ public:
 	FILEUTILITIES_API bool IsValid() const;
 
 	FILEUTILITIES_API TArray<FString> GetFileNames() const;
-	FILEUTILITIES_API bool TryReadFile(FStringView FileName, TArray<uint8>& OutData, FOutputDevice* ErrorHandler = nullptr) const;
+	FILEUTILITIES_API bool TryReadFile(FStringView FileName, TArray<uint8>& OutData, 
+		FOutputDevice* ErrorHandler = nullptr, 
+		UE::FileUtilities::FZipFileMetaData* OutMetadata = nullptr
+		) const;
 };
 
 #endif
diff --git a/Engine/Source/Developer/FileUtilities/Public/FileUtilities/ZipFileMetaData.h b/Engine/Source/Developer/FileUtilities/Public/FileUtilities/ZipFileMetaData.h
new file mode 100644
--- /dev/null
+++ b/Engine/Source/Developer/FileUtilities/Public/FileUtilities/ZipFileMetaData.h
@@ -0,0 +1,15 @@
+﻿// Copyright Epic Games, Inc. All Rights Reserved.
+
+#pragma once
+
+#include "Misc/DateTime.h"
+#include "Misc/Optional.h"
+
+namespace UE::FileUtilities
+{
+struct FZipFileMetaData
+{
+	/** The timestamp of the file, if available. */
+	TOptional<FDateTime> Timestamp;
+};
+}
