File Sandbox: Log last error when persist I/O operation fails.

#Synced-CL 52378315

diff --git a/Engine/Plugins/Developer/Sandbox/FileSandbox/Source/FileSandboxCore/Private/Sandbox/Persist/LogPersistFeedback.h b/Engine/Plugins/Developer/Sandbox/FileSandbox/Source/FileSandboxCore/Private/Sandbox/Persist/LogPersistFeedback.h
--- a/Engine/Plugins/Developer/Sandbox/FileSandbox/Source/FileSandboxCore/Private/Sandbox/Persist/LogPersistFeedback.h
+++ b/Engine/Plugins/Developer/Sandbox/FileSandbox/Source/FileSandboxCore/Private/Sandbox/Persist/LogPersistFeedback.h
@@ -2,6 +2,7 @@
 
 #pragma once 
 
+#include "HAL/PlatformMisc.h"
 #include "Interface/IPersistFeedback.h"
 #include "LogFileSandbox.h"
 #include "Logging/LogMacros.h"
@@ -43,16 +44,28 @@ public:
 	}
 	virtual void HandleError_MakeWritable(const TCHAR* InFilename) override
 	{
-		UE_LOG(LogFileSandbox, Warning, TEXT("Failed to make '%s' writeable using file system when persisting sandbox state!"), InFilename);
+		UE_LOG(LogFileSandbox, Warning, TEXT("Failed to make '%s' writeable using file system when persisting sandbox state! %s"), InFilename, *GetErrorMessage());
 	}
 	virtual void HandleError_MoveFile(const TCHAR* InToFilename, const TCHAR* InFromFilename) override
 	{
-		UE_LOG(LogFileSandbox, Warning, TEXT("Failed to move file '%s' to '%s' using file system when persisting sandbox state!"), InFromFilename, InToFilename);
+		const uint32 LastErrorCode = FPlatformMisc::GetLastError();
+		UE_LOG(LogFileSandbox, Warning, TEXT("Failed to move file '%s' to '%s' using file system when persisting sandbox state! %s"), InFromFilename, InToFilename, *GetErrorMessage());
 	}
 	virtual void HandleError_DeleteFile(const TCHAR* InFilename) override
 	{
-		UE_LOG(LogFileSandbox, Warning, TEXT("Failed to delete file '%s' using file system when persisting sandbox state!"), InFilename);
+		const uint32 LastErrorCode = FPlatformMisc::GetLastError();
+		UE_LOG(LogFileSandbox, Warning, TEXT("Failed to delete file '%s' using file system when persisting sandbox state! %s"), InFilename, *GetErrorMessage());
 	}
 	//~ End IPersistFeedback Interface
+	
+private:
+	
+	static FString GetErrorMessage()
+	{
+		const uint32 LastErrorCode = FPlatformMisc::GetLastError();
+		TCHAR ErrorBuffer[1024];
+		FPlatformMisc::GetSystemErrorMessage(ErrorBuffer, 1024, LastErrorCode);
+		return FString::Printf(TEXT("GetLastError: %d - %s"), LastErrorCode, ErrorBuffer);
+	}
 };
 }
