Sandboxed Editing: Fix UI show wrong number of files persisted

#Synced-CL 52378315

diff --git a/Engine/Plugins/Developer/Sandbox/FileSandbox/Source/FileSandboxUI/Private/Persist/Feedback/SummaryPersistFeedback.cpp b/Engine/Plugins/Developer/Sandbox/FileSandbox/Source/FileSandboxUI/Private/Persist/Feedback/SummaryPersistFeedback.cpp
--- a/Engine/Plugins/Developer/Sandbox/FileSandbox/Source/FileSandboxUI/Private/Persist/Feedback/SummaryPersistFeedback.cpp
+++ b/Engine/Plugins/Developer/Sandbox/FileSandbox/Source/FileSandboxUI/Private/Persist/Feedback/SummaryPersistFeedback.cpp
@@ -17,28 +17,12 @@ FSummaryPersistFeedback::FSummaryPersistFeedback(
 	{
 		const FileSandboxCore::ESandboxFileChange Change = InFileActions[Index];
 		FileActions.Add(FPaths::ConvertRelativePathToFull(InPersistedFiles[Index]), Change);
-		
-		// We'll decrement it again if an error occurs. We assume errors should be rare.
-		switch (Change)
-		{
-		case FileSandboxCore::ESandboxFileChange::Added: 
-			++Summary.Added;
-			break;
-		case FileSandboxCore::ESandboxFileChange::Removed:
-			++Summary.Deleted;
-			break;
-		case FileSandboxCore::ESandboxFileChange::Edited:
-			++Summary.Edited;
-			break;
-		case FileSandboxCore::ESandboxFileChange::None: [[fallthrough]];
-		default: ensure(false);
-		}
 	}
 }
 
-void FSummaryPersistFeedback::HandleError(const TCHAR* InToFilename)
+void FSummaryPersistFeedback::HandleSuccess(const TCHAR* InFilename)
 {
-	FString Abs = FPaths::ConvertRelativePathToFull(InToFilename);
+	FString Abs = FPaths::ConvertRelativePathToFull(InFilename);
 	
 	const FileSandboxCore::ESandboxFileChange* FileChange = FileActions.Find(Abs);
 	if (ensure(FileChange))
@@ -46,19 +30,23 @@ void FSummaryPersistFeedback::HandleError(const TCHAR* InToFilename)
 		switch (*FileChange)
 		{
 		case FileSandboxCore::ESandboxFileChange::Added: 
-			--Summary.Added;
+			++Summary.Added;
 			break;
 		case FileSandboxCore::ESandboxFileChange::Removed:
-			--Summary.Deleted;
+			++Summary.Deleted;
 			break;
 		case FileSandboxCore::ESandboxFileChange::Edited:
-			--Summary.Edited;
+			++Summary.Edited;
 			break;
 		case FileSandboxCore::ESandboxFileChange::None: [[fallthrough]];
 		default: ensure(false);
 		}
 	}
-	
+}
+
+void FSummaryPersistFeedback::HandleError(const TCHAR* InToFilename)
+{
+	FString Abs = FPaths::ConvertRelativePathToFull(InToFilename);
 	Summary.FailedFiles.Emplace(MoveTemp(Abs));
 }
 }
diff --git a/Engine/Plugins/Developer/Sandbox/FileSandbox/Source/FileSandboxUI/Public/Persist/Feedback/SummaryPersistFeedback.h b/Engine/Plugins/Developer/Sandbox/FileSandbox/Source/FileSandboxUI/Public/Persist/Feedback/SummaryPersistFeedback.h
--- a/Engine/Plugins/Developer/Sandbox/FileSandbox/Source/FileSandboxUI/Public/Persist/Feedback/SummaryPersistFeedback.h
+++ b/Engine/Plugins/Developer/Sandbox/FileSandbox/Source/FileSandboxUI/Public/Persist/Feedback/SummaryPersistFeedback.h
@@ -36,7 +36,7 @@ public:
 
 	//~ Begin IPersistFeedback Interface
 	virtual void StartFile(const TCHAR* InFilename) override { }
-	virtual void HandleSuccess(const TCHAR* InFilename) override { }
+	virtual void HandleSuccess(const TCHAR* InFilename) override;
 	virtual void HandleError_CheckoutNotAllowed(const TCHAR* InFilename) override { HandleError(InFilename); }
 	virtual void HandleError_Checkout(const TCHAR* InFilename) override { HandleError(InFilename); }
 	virtual void HandleError_Revert(const TCHAR* InFilename) override { HandleError(InFilename); }
diff --git a/Engine/Plugins/Developer/Sandbox/SandboxedEditing/Source/SandboxedEditing/Private/Features/Browser/Views/PersistSandboxView.cpp b/Engine/Plugins/Developer/Sandbox/SandboxedEditing/Source/SandboxedEditing/Private/Features/Browser/Views/PersistSandboxView.cpp
--- a/Engine/Plugins/Developer/Sandbox/SandboxedEditing/Source/SandboxedEditing/Private/Features/Browser/Views/PersistSandboxView.cpp
+++ b/Engine/Plugins/Developer/Sandbox/SandboxedEditing/Source/SandboxedEditing/Private/Features/Browser/Views/PersistSandboxView.cpp
@@ -86,7 +86,7 @@ namespace PersistViewDetail
 {
 static FText BuildErrorSubtext(const FileSandboxUI::FPersistSummary& InSummary)
 {
-	return FText::Format(LOCTEXT("Persist.SubText.WithErrorsFmt", "{0} persisted. {1} errors.\nSee log for details."), 
+	return FText::Format(LOCTEXT("Persist.SubText.WithErrorsFmt", "{0} persisted. {1} {0}|plural(one=error,other=errors).\nSee log for details."), 
 		InSummary.Edited + InSummary.Added + InSummary.Deleted,
 		InSummary.FailedFiles.Num()
 		);
@@ -95,9 +95,9 @@ static FText BuildSuccessSubtext(const FileSandboxUI::FPersistSummary& InSummary
 {
 	return FText::Format(
 		LOCTEXT("Persist.SubText.SuccessFmt", "{0}{1}{2}"),
-		InSummary.Edited > 0 ? FText::Format(LOCTEXT("Edited", "Edited {0} files.\n"), InSummary.Edited) : FText::GetEmpty(),
-		InSummary.Added > 0 ? FText::Format(LOCTEXT("Added", "Added {0} files.\n"), InSummary.Added) : FText::GetEmpty(),
-		InSummary.Deleted > 0 ? FText::Format(LOCTEXT("Removed", "Removed {0} files."), InSummary.Deleted) : FText::GetEmpty()
+		InSummary.Edited > 0 ? FText::Format(LOCTEXT("Edited", "Edited {0} {0}|plural(one=file,other=files).\n"), InSummary.Edited) : FText::GetEmpty(),
+		InSummary.Added > 0 ? FText::Format(LOCTEXT("Added", "Added {0} {0}|plural(one=file,other=files).\n"), InSummary.Added) : FText::GetEmpty(),
+		InSummary.Deleted > 0 ? FText::Format(LOCTEXT("Removed", "Removed {0} {0}|plural(one=file,other=files)."), InSummary.Deleted) : FText::GetEmpty()
 		);
 }
 }
