File Sandbox: Fix missing dll export

#Synced-CL 52378315

diff --git a/Engine/Plugins/Developer/Sandbox/FileSandbox/Source/FileSandboxUI/Public/Persist/Feedback/SummaryPersistFeedback.h b/Engine/Plugins/Developer/Sandbox/FileSandbox/Source/FileSandboxUI/Public/Persist/Feedback/SummaryPersistFeedback.h
--- a/Engine/Plugins/Developer/Sandbox/FileSandbox/Source/FileSandboxUI/Public/Persist/Feedback/SummaryPersistFeedback.h
+++ b/Engine/Plugins/Developer/Sandbox/FileSandbox/Source/FileSandboxUI/Public/Persist/Feedback/SummaryPersistFeedback.h
@@ -36,7 +36,7 @@ public:
 
 	//~ Begin IPersistFeedback Interface
 	virtual void StartFile(const TCHAR* InFilename) override { }
-	virtual void HandleSuccess(const TCHAR* InFilename) override;
+	UE_API virtual void HandleSuccess(const TCHAR* InFilename) override;
 	virtual void HandleError_CheckoutNotAllowed(const TCHAR* InFilename) override { HandleError(InFilename); }
 	virtual void HandleError_Checkout(const TCHAR* InFilename) override { HandleError(InFilename); }
 	virtual void HandleError_Revert(const TCHAR* InFilename) override { HandleError(InFilename); }
