Sandboxed Editing: Fix timestamp of exported files being set to UTC time

#Synced-CL 52378315

diff --git a/Engine/Plugins/Developer/Sandbox/FileSandbox/Source/FileSandboxCore/Private/Utils/ZipTransferUtils.cpp b/Engine/Plugins/Developer/Sandbox/FileSandbox/Source/FileSandboxCore/Private/Utils/ZipTransferUtils.cpp
--- a/Engine/Plugins/Developer/Sandbox/FileSandbox/Source/FileSandboxCore/Private/Utils/ZipTransferUtils.cpp
+++ b/Engine/Plugins/Developer/Sandbox/FileSandbox/Source/FileSandboxCore/Private/Utils/ZipTransferUtils.cpp
@@ -26,6 +26,14 @@ namespace UE::FileSandboxCore::ZipUtils
 {
 namespace ZipDetail
 {
+static FDateTime ToLocalTime(const FDateTime& UtcTime)
+{
+	const FDateTime UtcNow = FDateTime::UtcNow();
+	const FDateTime Now = FDateTime::Now();
+	const FTimespan TimeZoneOffset = Now - UtcNow;
+	return UtcTime + TimeZoneOffset;
+}
+
 static bool WriteToZip(const FString& InSandboxRoot, IFileHandle& InZipFile, TConstArrayView<FString> InFilesToCompress)
 {
 	IFileManager& FileManager = IFileManager::Get();
@@ -53,10 +61,12 @@ static bool WriteToZip(const FString& InSandboxRoot, IFileHandle& InZipFile, TCo
 		FPaths::MakePathRelativeTo(RelativeToRoot, *InSandboxRoot);
 		
 		// The file in the archive needs to have the same timestamp as in the file system so the edit time is correct when importing
-		const FDateTime Timestamp = FileManager.GetTimeStamp(*FileToCompress);
+		// File system returns time in UTC...
+		FDateTime Timestamp = ToLocalTime(FileManager.GetTimeStamp(*FileToCompress)); 
 		const bool bHasTimestamp = Timestamp != FDateTime::MinValue(); 
 		const FDateTime TimestampToSet = bHasTimestamp ? Timestamp : FDateTime::Now();
 		
+		// ... but ZIP API expects time in local time
 		ZipArchiveWriter.AddFile(RelativeToRoot, RawData, TimestampToSet);
 	}
 	
