UE-356522: Add version checking to prevent loading newer sandboxes

The UI now refuses to load a sandbox created with a newer version of
Unreal Engine. For example, you cannot load a 5.8 sandbox in 5.7, and
you cannot load a sandbox with CL 5819811 in CL 5719181.

The version check occurs BEFORE leaving the current sandbox, ensuring
that if the load fails due to incompatibility, the user remains in their
current sandbox without disruption.

Changes:
- Added IsCompatibleWith() method to FFileSandboxCore_EngineVersionInfo
- Added IncompatibleVersion error reason to ELoadSandboxLoadErrorReason
- Added version check in FSandboxManager::LoadSandbox() before LeaveSandbox()
- Added ShowIncompatibleVersionError() notification function
- Enhanced FSandboxSystemModel::LoadSandbox() to show specific error

#rnx
#Synced-CL 52474683

diff --git a/Engine/Plugins/Developer/Sandbox/FileSandbox/Source/FileSandboxCore/Private/Data/VersionInfo.cpp b/Engine/Plugins/Developer/Sandbox/FileSandbox/Source/FileSandboxCore/Private/Data/VersionInfo.cpp
--- a/Engine/Plugins/Developer/Sandbox/FileSandbox/Source/FileSandboxCore/Private/Data/VersionInfo.cpp
+++ b/Engine/Plugins/Developer/Sandbox/FileSandbox/Source/FileSandboxCore/Private/Data/VersionInfo.cpp
@@ -27,11 +27,48 @@ void FFileSandboxCore_EngineVersionInfo::Initialize(const FEngineVersion& InVers
 
 FString FFileSandboxCore_EngineVersionInfo::ToString(bool bIncludeChangelist) const
 {
-	return bIncludeChangelist 
+	return bIncludeChangelist
 		? FString::Printf(TEXT("%u.%u.%u-%u"), Major, Minor, Patch, Changelist)
 		: FString::Printf(TEXT("%u.%u.%u"), Major, Minor, Patch);
 }
 
+bool FFileSandboxCore_EngineVersionInfo::IsCompatibleWith(const FFileSandboxCore_EngineVersionInfo& InCurrentVersion) const
+{
+	// Check major version
+	if (Major > InCurrentVersion.Major)
+	{
+		return false;
+	}
+	if (Major < InCurrentVersion.Major)
+	{
+		return true;
+	}
+
+	// Major versions match, check minor
+	if (Minor > InCurrentVersion.Minor)
+	{
+		return false;
+	}
+	if (Minor < InCurrentVersion.Minor)
+	{
+		return true;
+	}
+
+	// Major and minor match, check patch
+	if (Patch > InCurrentVersion.Patch)
+	{
+		return false;
+	}
+	if (Patch < InCurrentVersion.Patch)
+	{
+		return true;
+	}
+
+	// Major, minor, and patch match - check changelist
+	// A sandbox with a higher changelist cannot be loaded in a lower changelist build
+	return Changelist <= InCurrentVersion.Changelist;
+}
+
 void FFileSandboxCore_CustomVersionInfo::Initialize(const FCustomVersion& InVersion)
 {
 	FriendlyName = InVersion.GetFriendlyName();
diff --git a/Engine/Plugins/Developer/Sandbox/FileSandbox/Source/FileSandboxCore/Private/Sandbox/SandboxManager.cpp b/Engine/Plugins/Developer/Sandbox/FileSandbox/Source/FileSandboxCore/Private/Sandbox/SandboxManager.cpp
--- a/Engine/Plugins/Developer/Sandbox/FileSandbox/Source/FileSandboxCore/Private/Sandbox/SandboxManager.cpp
+++ b/Engine/Plugins/Developer/Sandbox/FileSandbox/Source/FileSandboxCore/Private/Sandbox/SandboxManager.cpp
@@ -2,12 +2,14 @@
 
 #include "SandboxManager.h"
 
+#include "Data/VersionInfo.h"
 #include "IFileSandboxCoreModule.h"
 #include "ISandboxRepository.h"
 #include "GenericPlatform/GenericPlatformFile.h"
 #include "HAL/PlatformFileManager.h"
 #include "LogFileSandbox.h"
 #include "Misc/CoreDelegates.h"
+#include "Misc/EngineVersion.h"
 #include "Misc/Paths.h"
 #include "Persist/DiffUtils.h"
 #include "Platform/SanboxedFilePathUtils.h"
@@ -21,6 +23,7 @@
 #include "Types/Manager/NewSandboxArgs.h"
 #include "Types/Manager/NewSandboxError.h"
 #include "Types/Manager/SandboxCreationResult.h"
+#include "Utils/SandboxDirectoryUtils.h"
 #include "Utils/SandboxFileUtils.h"
 
 namespace UE::FileSandboxCore
@@ -73,18 +76,36 @@ FLoadSandboxResult FSandboxManager::LoadSandbox(const FLoadSandboxByDirectoryArg
 	}
 
 	// This could happen if the sandbox is placed in the root e.g. "D:/".
-	// It means that the sandbox was not created in GetBaseSandboxDirectory. 
+	// It means that the sandbox was not created in GetBaseSandboxDirectory.
 	// Names are just a simplified concept for managing sandboxes that are placed in this default folder.
 	// This is a custom sandbox so it's beyond the default workflow.
 	const FString LastDirectoryName = FPaths::GetPathLeaf(AbsoluteRoot);
 	UE_CLOG(LastDirectoryName.IsEmpty(), LogFileSandbox, Warning, TEXT("Sandbox placed %s will not have any name."), *AbsoluteRoot);
 	const TOptional<FString> SandboxName = LastDirectoryName.IsEmpty() ? LastDirectoryName : TOptional<FString>{};
-	
+
+	// Check version compatibility BEFORE leaving current sandbox (UE-356522)
+	// If the version is incompatible, we want to fail early without disrupting the user's current state
+	const FFileSandboxCore_VersionInfo SandboxVersionInfo = LoadVersionInfo(AbsoluteRoot);
+	if (SandboxVersionInfo.IsInitialized())
+	{
+		FFileSandboxCore_EngineVersionInfo CurrentVersion;
+		CurrentVersion.Initialize(FEngineVersion::Current());
+
+		if (!SandboxVersionInfo.EngineVersion.IsCompatibleWith(CurrentVersion))
+		{
+			UE_LOG(LogFileSandbox, Error, TEXT("Cannot load sandbox at %s: Sandbox was created with version %s but current version is %s"),
+				*AbsoluteRoot,
+				*SandboxVersionInfo.EngineVersion.ToString(),
+				*CurrentVersion.ToString());
+			return { MakeError(FLoadSandboxError{ ELoadSandboxLoadErrorReason::IncompatibleVersion }) };
+		}
+	}
+
 	if (!LeaveSandbox())
 	{
 		return { MakeError(FLoadSandboxError{ ELoadSandboxLoadErrorReason::CannotLeaveSandbox }) };
 	}
-	
+
 	TUniquePtr<FSandboxInstance> NewSandbox = FSandboxInstance::LoadSandbox(AbsoluteRoot, InLoadArgs.InitArgs);
 	ISandboxInstance* SandboxRaw = NewSandbox.Get();
 	if (!SandboxRaw)
diff --git a/Engine/Plugins/Developer/Sandbox/FileSandbox/Source/FileSandboxCore/Private/Types/Manager/LoadSandboxError.cpp b/Engine/Plugins/Developer/Sandbox/FileSandbox/Source/FileSandboxCore/Private/Types/Manager/LoadSandboxError.cpp
--- a/Engine/Plugins/Developer/Sandbox/FileSandbox/Source/FileSandboxCore/Private/Types/Manager/LoadSandboxError.cpp
+++ b/Engine/Plugins/Developer/Sandbox/FileSandbox/Source/FileSandboxCore/Private/Types/Manager/LoadSandboxError.cpp
@@ -8,13 +8,14 @@ namespace UE::FileSandboxCore
 {
 FString LexToString(ELoadSandboxLoadErrorReason InReason)
 {
-	static_assert(static_cast<int32>(ELoadSandboxLoadErrorReason::Count) == 5, "Update this switch");
+	static_assert(static_cast<int32>(ELoadSandboxLoadErrorReason::Count) == 6, "Update this switch");
 	switch (InReason)
 	{
 	case ELoadSandboxLoadErrorReason::InvalidDirectory: return TEXT("InvalidDirectory");
 	case ELoadSandboxLoadErrorReason::InvalidName: return TEXT("InvalidName");
 	case ELoadSandboxLoadErrorReason::IOError: return TEXT("IOError");
 	case ELoadSandboxLoadErrorReason::CannotLeaveSandbox: return TEXT("CannotLeaveSandbox");
+	case ELoadSandboxLoadErrorReason::IncompatibleVersion: return TEXT("IncompatibleVersion");
 	case ELoadSandboxLoadErrorReason::Unspecified: return TEXT("Unspecified");
 	default: checkNoEntry(); return TEXT("Unknown");
 	}
diff --git a/Engine/Plugins/Developer/Sandbox/FileSandbox/Source/FileSandboxCore/Public/Data/VersionInfo.h b/Engine/Plugins/Developer/Sandbox/FileSandbox/Source/FileSandboxCore/Public/Data/VersionInfo.h
--- a/Engine/Plugins/Developer/Sandbox/FileSandbox/Source/FileSandboxCore/Public/Data/VersionInfo.h
+++ b/Engine/Plugins/Developer/Sandbox/FileSandbox/Source/FileSandboxCore/Public/Data/VersionInfo.h
@@ -54,6 +54,15 @@ struct FFileSandboxCore_EngineVersionInfo
 
 	/** @return Engine version info as string, e.g. "5.7.1-48893048" */
 	FILESANDBOXCORE_API FString ToString(bool bIncludeChangelist = true) const;
+
+	/**
+	 * Check if this version is compatible with the given version.
+	 * A sandbox version is compatible if it's not newer than the current version.
+	 * Comparison order: Major > Minor > Patch > Changelist
+	 * @param InCurrentVersion The current engine version to check against
+	 * @return true if this version can be loaded in the current version
+	 */
+	FILESANDBOXCORE_API bool IsCompatibleWith(const FFileSandboxCore_EngineVersionInfo& InCurrentVersion) const;
 };
 
 /** Holds custom version information */
diff --git a/Engine/Plugins/Developer/Sandbox/FileSandbox/Source/FileSandboxCore/Public/Types/Manager/LoadSandboxError.h b/Engine/Plugins/Developer/Sandbox/FileSandbox/Source/FileSandboxCore/Public/Types/Manager/LoadSandboxError.h
--- a/Engine/Plugins/Developer/Sandbox/FileSandbox/Source/FileSandboxCore/Public/Types/Manager/LoadSandboxError.h
+++ b/Engine/Plugins/Developer/Sandbox/FileSandbox/Source/FileSandboxCore/Public/Types/Manager/LoadSandboxError.h
@@ -26,12 +26,19 @@ enum class ELoadSandboxLoadErrorReason : uint8
 	 */
 	CannotLeaveSandbox,
 
+	/**
+	 * The sandbox was created with a newer version of Unreal Engine and cannot be loaded.
+	 * For example, a 5.8 sandbox cannot be loaded in 5.7, or a sandbox with CL 5819811
+	 * cannot be loaded in CL 5719181.
+	 */
+	IncompatibleVersion,
+
 	/** Dummy error code for initializing to a deterministic value. */
 	Unspecified,
-	
-	/** 
-	 * ADD NEW MEMBERS ABOVE. 
-	 * Not a real reason. 
+
+	/**
+	 * ADD NEW MEMBERS ABOVE.
+	 * Not a real reason.
 	 */
 	Count
 };
diff --git a/Engine/Plugins/Developer/Sandbox/SandboxedEditing/Source/SandboxedEditing/Private/Framework/Models/SandboxSystemModel.cpp b/Engine/Plugins/Developer/Sandbox/SandboxedEditing/Source/SandboxedEditing/Private/Framework/Models/SandboxSystemModel.cpp
--- a/Engine/Plugins/Developer/Sandbox/SandboxedEditing/Source/SandboxedEditing/Private/Framework/Models/SandboxSystemModel.cpp
+++ b/Engine/Plugins/Developer/Sandbox/SandboxedEditing/Source/SandboxedEditing/Private/Framework/Models/SandboxSystemModel.cpp
@@ -5,13 +5,16 @@
 #include "AssetViewUtils.h"
 #include "Containers/UnrealString.h"
 #include "Data/SandboxMetaData.h"
+#include "Data/VersionInfo.h"
 #include "Framework/Models/CreateSandboxArgs.h"
 #include "Framework/Models/SandboxInfo.h"
+#include "Framework/Notifications.h"
 #include "IFileSandboxCoreModule.h"
 #include "ISandboxInstance.h"
 #include "ISandboxManager.h"
 #include "ISandboxRepository.h"
 #include "Interface/ISandboxLock.h"
+#include "Misc/EngineVersion.h"
 #include "Misc/Paths.h"
 #include "SandboxedEditingSettings.h"
 #include "SandboxTagging.h"
@@ -20,6 +23,7 @@
 #include "Types/Manager/DeleteSandboxByDirectoryArgs.h"
 #include "Types/Manager/DeleteSandboxResult.h"
 #include "Types/Manager/LoadSandboxByDirectoryArgs.h"
+#include "Types/Manager/LoadSandboxError.h"
 #include "Types/Manager/NewSandboxArgs.h"
 #include "Types/Sandbox/RevertResult.h"
 #include "Types/SandboxMetaInfo.h"
@@ -128,7 +132,36 @@ bool FSandboxSystemModel::LoadSandbox(const FString& InRootDirectory)
 		using namespace FileSandboxCore;
 		IFileSandboxCoreModule& Module = IFileSandboxCoreModule::Get();
 		ISandboxManager& SandboxManager = Module.GetSandboxManager();
-		return SandboxManager.LoadSandbox(FLoadSandboxByDirectoryArgs{ InRootDirectory }).HasValue();
+
+		FLoadSandboxResult Result = SandboxManager.LoadSandbox(FLoadSandboxByDirectoryArgs{ InRootDirectory });
+		if (Result.HasValue())
+		{
+			return true;
+		}
+
+		// Handle specific error cases (UE-356522)
+		if (Result.HasError())
+		{
+			const FLoadSandboxError& Error = Result.GetError();
+			if (Error.Reason == ELoadSandboxLoadErrorReason::IncompatibleVersion)
+			{
+				// Get sandbox name and version info for the error message
+				const FString SandboxName = FPaths::GetPathLeaf(InRootDirectory);
+				const FFileSandboxCore_VersionInfo SandboxVersion = LoadVersionInfo(InRootDirectory);
+
+				FFileSandboxCore_EngineVersionInfo CurrentVersion;
+				CurrentVersion.Initialize(FEngineVersion::Current());
+
+				ShowIncompatibleVersionError(
+					SandboxName,
+					SandboxVersion.EngineVersion.ToString(),
+					CurrentVersion.ToString()
+				);
+				return false;
+			}
+		}
+
+		// Fall through to generic failure for other errors
 	}
 	return false;
 }
diff --git a/Engine/Plugins/Developer/Sandbox/SandboxedEditing/Source/SandboxedEditing/Private/Framework/Notifications.cpp b/Engine/Plugins/Developer/Sandbox/SandboxedEditing/Source/SandboxedEditing/Private/Framework/Notifications.cpp
--- a/Engine/Plugins/Developer/Sandbox/SandboxedEditing/Source/SandboxedEditing/Private/Framework/Notifications.cpp
+++ b/Engine/Plugins/Developer/Sandbox/SandboxedEditing/Source/SandboxedEditing/Private/Framework/Notifications.cpp
@@ -50,19 +50,44 @@ void ShowFailedToLoadSandbox(const FString& InSandboxName)
 	{
 		return;
 	}
-	
+
 	FSlateNotificationManager& NotificationManager = FSlateNotificationManager::Get();
 	FNotificationInfo Args(
 		FText::Format(LOCTEXT("FailedToLoadFmt.Title", "Failed to load sandbox {0}"), FText::FromString(InSandboxName))
 		);
 	Args.bFireAndForget = true;
 	Args.ExpireDuration = 2.f;
-	
+
 	const TSharedPtr<SNotificationItem> Item = NotificationManager.AddNotification(Args);
 	Item->SetCompletionState(SNotificationItem::CS_Fail);
 	Item->SetSubText(LOCTEXT("FailedToLoadFmt.Description", "See log for more info."));
 }
 
+void ShowIncompatibleVersionError(const FString& InSandboxName, const FString& InSandboxVersion, const FString& InCurrentVersion)
+{
+	if (!FSlateApplication::IsInitialized())
+	{
+		return;
+	}
+
+	FSlateNotificationManager& NotificationManager = FSlateNotificationManager::Get();
+	FNotificationInfo Args(
+		FText::Format(LOCTEXT("IncompatibleVersionFmt.Title", "Cannot load sandbox {0}"), FText::FromString(InSandboxName))
+		);
+	Args.bFireAndForget = true;
+	Args.ExpireDuration = 5.f;
+
+	const TSharedPtr<SNotificationItem> Item = NotificationManager.AddNotification(Args);
+	Item->SetCompletionState(SNotificationItem::CS_Fail);
+	Item->SetSubText(
+		FText::Format(
+			LOCTEXT("IncompatibleVersionFmt.Description", "Sandbox was created with version {0}, but you are running version {1}. You cannot load a sandbox created with a newer version."),
+			FText::FromString(InSandboxVersion),
+			FText::FromString(InCurrentVersion)
+		)
+	);
+}
+
 void ShowLeftSandbox(const FString& InSandboxName)
 {
 	if (!FSlateApplication::IsInitialized())
diff --git a/Engine/Plugins/Developer/Sandbox/SandboxedEditing/Source/SandboxedEditing/Private/Framework/Notifications.h b/Engine/Plugins/Developer/Sandbox/SandboxedEditing/Source/SandboxedEditing/Private/Framework/Notifications.h
--- a/Engine/Plugins/Developer/Sandbox/SandboxedEditing/Source/SandboxedEditing/Private/Framework/Notifications.h
+++ b/Engine/Plugins/Developer/Sandbox/SandboxedEditing/Source/SandboxedEditing/Private/Framework/Notifications.h
@@ -17,6 +17,8 @@ void ShowCreatedSandbox(const FString& InSandboxName);
 void ShowLoadedSandbox(const FString& InSandboxName);
 /** Show that the sandbox could not be loaded. */
 void ShowFailedToLoadSandbox(const FString& InSandboxName);
+/** Show that the sandbox could not be loaded due to incompatible version. */
+void ShowIncompatibleVersionError(const FString& InSandboxName, const FString& InSandboxVersion, const FString& InCurrentVersion);
 
 /** Shows notification that the user left a sandbox. */
 void ShowLeftSandbox(const FString& InSandboxName);
