Add sandbox name length validation to prevent path length issues

The default limit is 64 characters, configurable via the `Sandbox.MaximumSandboxNameLength` console variable (valid range: 32-1024). When users attempt to create sandboxes with names exceeding the limit, they receive an error showing the current length and maximum allowed.

#Synced-CL 52563198

diff --git a/Engine/Plugins/Developer/Sandbox/FileSandbox/Source/FileSandboxCore/Private/Utils/SandboxDirectoryUtils.cpp b/Engine/Plugins/Developer/Sandbox/FileSandbox/Source/FileSandboxCore/Private/Utils/SandboxDirectoryUtils.cpp
--- a/Engine/Plugins/Developer/Sandbox/FileSandbox/Source/FileSandboxCore/Private/Utils/SandboxDirectoryUtils.cpp
+++ b/Engine/Plugins/Developer/Sandbox/FileSandbox/Source/FileSandboxCore/Private/Utils/SandboxDirectoryUtils.cpp
@@ -1,4 +1,4 @@
-﻿// Copyright Epic Games, Inc. All Rights Reserved.
+// Copyright Epic Games, Inc. All Rights Reserved.
 
 #include "Utils/SandboxDirectoryUtils.h"
 
@@ -20,11 +20,38 @@
 #include "Sandbox/Platform/SandboxedPlatformFilePath.h"
 #include "Templates/Function.h"
 #include "Types/EBreakBehavior.h"
+#include "HAL/IConsoleManager.h"
 
 #define LOCTEXT_NAMESPACE "SandboxDirectoryUtils"
 
 namespace UE::FileSandboxCore
 {
+
+/** Minimum allowed sandbox name length */
+constexpr int32 MinSandboxNameLength = 32;
+
+/** Maximum allowed sandbox name length */
+constexpr int32 MaxSandboxNameLength = 1024;
+
+/** Console variable to control maximum sandbox name length (clamped between 32 and 1024 characters) */
+static TAutoConsoleVariable<int32> CVarMaximumSandboxNameLength(
+	TEXT("Sandbox.MaximumSandboxNameLength"),
+	64,
+	TEXT("Maximum allowed length for sandbox names. Valid range: 32-1024 characters.\n")
+	TEXT("Longer names may cause filesystem path length issues."),
+	ECVF_Default
+);
+
+/**
+ * Gets the maximum sandbox name length from the console variable, clamped to valid range [32, 1024].
+ * @return Maximum sandbox name length
+ */
+int32 GetMaximumSandboxNameLength()
+{
+	const int32 Value = CVarMaximumSandboxNameLength.GetValueOnAnyThread();
+	return FMath::Clamp(Value, MinSandboxNameLength, MaxSandboxNameLength);
+}
+	
 FString GetBaseSandboxDirectory()
 {
 	const FString CommandLineOverride = ParseDefaultSandboxDirectory();
@@ -232,7 +259,14 @@ EDirectorySuitability DetermineDirectorySuitability(const FString& InSandboxName
 	{
 		return EDirectorySuitability::EmptyName;
 	}
-	
+
+	// Limit sandbox name length to prevent excessively long directory paths and FName character limits
+	const int32 MaxLength = GetMaximumSandboxNameLength();
+	if (InSandboxName.Len() > MaxLength)
+	{
+		return EDirectorySuitability::NameTooLong;
+	}
+
 	const FString NewDirectory = FPaths::Combine(InBaseDirectory, InSandboxName);
 	if (!FPaths::ValidatePath(NewDirectory))
 	{
@@ -243,7 +277,7 @@ EDirectorySuitability DetermineDirectorySuitability(const FString& InSandboxName
 	{
 		return EDirectorySuitability::DuplicateName;
 	}
-	
+
 	// We'll require that a sandbox be created in a new directory so we don't have to worry handling pre-existing files in that directory.
 	IPlatformFile& PlatformFile = FPlatformFileManager::Get().GetPlatformFile();
 	const bool bDirectoryAlreadyExists = PlatformFile.DirectoryExists(*NewDirectory);
@@ -251,7 +285,7 @@ EDirectorySuitability DetermineDirectorySuitability(const FString& InSandboxName
 	{
 		return EDirectorySuitability::DirectoryAlreadyExists;
 	}
-	
+
 	return EDirectorySuitability::Suitable;
 }
 
@@ -261,21 +295,27 @@ FText FormatDirectorySuitabilityAsReason(EDirectorySuitability Suitability, cons
 	{
 	case EDirectorySuitability::Suitable: return LOCTEXT("Suitability.Suitable", "Suitable");
 	case EDirectorySuitability::EmptyName: return LOCTEXT("Suitability.EmptyName", "Sandbox name cannot be empty");
+	case EDirectorySuitability::NameTooLong:
+		return FText::Format(
+			LOCTEXT("Suitability.NameTooLong", "Sandbox name is too long ({0} characters). Maximum length is {1} characters."),
+			FText::AsNumber(InSandboxName.Len()),
+			FText::AsNumber(GetMaximumSandboxNameLength())
+			);
 	case EDirectorySuitability::DuplicateName:
 		return FText::Format(LOCTEXT("Suitability.DuplicateName", "Sandbox name {0} already exists"), FText::FromString(InSandboxName));
-		
+
 	case EDirectorySuitability::InvalidPath:
 		return FText::Format(
 			LOCTEXT("Suitability.InvalidPath", "Directory {0} is malformed"),
 			FText::FromString(FPaths::Combine(InBaseDirectory, InSandboxName))
 			);
-		
+
 	case EDirectorySuitability::DirectoryAlreadyExists:
 		return FText::Format(
 			LOCTEXT("Suitability.DirectoryAlreadyExists", "Directory {0} already exists"),
 			FText::FromString(FPaths::Combine(InBaseDirectory, InSandboxName))
 			);
-		
+
 	default: checkNoEntry(); return FText::GetEmpty();
 	}
 }
@@ -324,4 +364,4 @@ TOptional<FString> GetSandboxName(const FString& InSandboxDirectory)
 }
 }
 
-#undef LOCTEXT_NAMESPACE
\ No newline at end of file
+#undef LOCTEXT_NAMESPACE
diff --git a/Engine/Plugins/Developer/Sandbox/FileSandbox/Source/FileSandboxCore/Public/Utils/SandboxDirectoryUtils.h b/Engine/Plugins/Developer/Sandbox/FileSandbox/Source/FileSandboxCore/Public/Utils/SandboxDirectoryUtils.h
--- a/Engine/Plugins/Developer/Sandbox/FileSandbox/Source/FileSandboxCore/Public/Utils/SandboxDirectoryUtils.h
+++ b/Engine/Plugins/Developer/Sandbox/FileSandbox/Source/FileSandboxCore/Public/Utils/SandboxDirectoryUtils.h
@@ -124,6 +124,7 @@ enum class EDirectorySuitability : uint8
 {
 	Suitable,
 	EmptyName,
+	NameTooLong,
 	DuplicateName,
 	InvalidPath,
 	DirectoryAlreadyExists
