File Sandbox: Fix file disappearing from content browser after persist.

This also fixes the file from being marked for removed after duplication.

#Synced-CL 52617198

diff --git a/Engine/Plugins/Developer/Sandbox/FileSandbox/Source/FileSandboxCore/Private/Sandbox/Platform/SandboxPlatformFile.cpp b/Engine/Plugins/Developer/Sandbox/FileSandbox/Source/FileSandboxCore/Private/Sandbox/Platform/SandboxPlatformFile.cpp
--- a/Engine/Plugins/Developer/Sandbox/FileSandbox/Source/FileSandboxCore/Private/Sandbox/Platform/SandboxPlatformFile.cpp
+++ b/Engine/Plugins/Developer/Sandbox/FileSandbox/Source/FileSandboxCore/Private/Sandbox/Platform/SandboxPlatformFile.cpp
@@ -1203,7 +1203,16 @@ void FSandboxPlatformFile::OnSandboxDirectoryChanged(const TArray<FFileChangeDat
 			const FSandboxedPlatformFilePath RemappedFilePath = FSandboxedPlatformFilePath::CreateNonSandboxPath(
 				FPaths::ConvertRelativePathToFull(FileChange.Filename), MountPath
 				);
-			RemappedFileChanges.Emplace(RemappedFilePath.GetNonSandboxPath(), FileChange.Action);
+
+			// If the sandbox file was removed but the non-sandbox file still exists, this is a
+			// persist or edit-revert operation, not a true deletion. Emit FCA_Modified so the
+			// asset registry rescans the file rather than removing it from the content browser.
+			const FFileChangeData::EFileChangeAction RemappedAction =
+				(FileChange.Action == FFileChangeData::FCA_Removed && LowerLevel->FileExists(*RemappedFilePath.GetNonSandboxPath()))
+				? FFileChangeData::FCA_Modified
+				: FileChange.Action;
+
+			RemappedFileChanges.Emplace(RemappedFilePath.GetNonSandboxPath(), RemappedAction);
 
 			// Make sure the deleted state of this item is synchronized correctly
 			if (FileChange.Action == FFileChangeData::FCA_Added)
