File Sandbox: Fix not being able to persist files in Lyra. The fix is to flush the file perform attempting a move operation. The issue was the FLinkedLoad was still holding on to the file handle, which causes GetLastError 32.

#Synced-CL 52617198

diff --git a/Engine/Plugins/Developer/Sandbox/FileSandbox/Source/FileSandboxCore/Private/Sandbox/Persist/PersistUtils.cpp b/Engine/Plugins/Developer/Sandbox/FileSandbox/Source/FileSandboxCore/Private/Sandbox/Persist/PersistUtils.cpp
--- a/Engine/Plugins/Developer/Sandbox/FileSandbox/Source/FileSandboxCore/Private/Sandbox/Persist/PersistUtils.cpp
+++ b/Engine/Plugins/Developer/Sandbox/FileSandbox/Source/FileSandboxCore/Private/Sandbox/Persist/PersistUtils.cpp
@@ -118,7 +118,7 @@ bool DeleteFileWithSCC(
 {
 	// If this file maps to a package then we need to flush its linker so that we can remove the file from disk
 	FlushPackageFile(InFilename);
-
+	
 	// Get the source control state of the file
 	const FSourceControlStatePtr FileSCCState = InSourceControlProvider.IsEnabled() 
 		? InSourceControlProvider.GetState(InFilename, EStateCacheUsage::ForceUpdate) : nullptr;
@@ -167,6 +167,9 @@ static bool MoveFile(IPlatformFile& InPlatformFile, IPersistFeedback& InErrorHan
 {
 	bool bSuccess = true;
 		
+	// If this file maps to a package then we need to flush its linker so that we can remove the file from disk
+	FlushPackageFile(InToFilename);
+	
 	// MoveFile fails if there already is a target file
 	if (InPlatformFile.FileExists(InToFilename))
 	{
@@ -174,7 +177,7 @@ static bool MoveFile(IPlatformFile& InPlatformFile, IPersistFeedback& InErrorHan
 		bSuccess &= !InPlatformFile.IsReadOnly(InToFilename) || InPlatformFile.SetReadOnly(InToFilename, false);
 		bSuccess &= InPlatformFile.DeleteFile(InToFilename);
 	}
-		
+	
 	bSuccess &= InPlatformFile.MoveFile(InToFilename, InFromFilename);
 	
 	if (!bSuccess)
